[UBAController]
* Added plugin used to run distributed shader compiling through uba/horde

diff --git a/Engine/Plugins/UbaController/Source/UbaController/Private/UbaControllerModule.cpp b/Engine/Plugins/UbaController/Source/UbaController/Private/UbaControllerModule.cpp
new file mode 100644
--- /dev/null
+++ b/Engine/Plugins/UbaController/Source/UbaController/Private/UbaControllerModule.cpp
@@ -0,0 +1,227 @@
+// Copyright Epic Games, Inc. All Rights Reserved.
+
+#include "UbaControllerModule.h"
+
+#include "Misc/ConfigCacheIni.h"
+#include "UbaJobProcessor.h"
+#include "Features/IModularFeatures.h"
+#include "HAL/FileManager.h"
+#include "Misc/CommandLine.h"
+#include "Misc/CoreDelegates.h"
+#include "Misc/Paths.h"
+#include "Modules/ModuleManager.h"
+#include "CoreMinimal.h"
+
+
+DEFINE_LOG_CATEGORY(LogUbaController);
+
+namespace UbaControllerModule
+{
+	static constexpr int32 SubFolderCount = 32;
+};
+
+FUbaControllerModule::FUbaControllerModule()
+	: bSupported(false)
+	, bModuleInitialized(false)
+	, bControllerInitialized(false)
+	, RootWorkingDirectory(FPaths::Combine(FPlatformProcess::UserTempDir(), TEXT("UbaWorkingDir")))
+	, WorkingDirectory(FPaths::Combine(RootWorkingDirectory, FGuid::NewGuid().ToString(EGuidFormats::Digits)))
+	, NextFileID(0)
+	, NextTaskID(0)
+{
+}
+
+FUbaControllerModule::~FUbaControllerModule()
+{	
+	if (JobDispatcherThread)
+	{
+		JobDispatcherThread->Stop();
+		// Wait until the thread is done
+		FPlatformProcess::ConditionalSleep([&](){ return JobDispatcherThread && JobDispatcherThread->IsWorkDone(); },0.1f);
+	}
+
+	CleanWorkingDirectory();
+}
+
+bool FUbaControllerModule::IsSupported()
+{
+	if (bControllerInitialized)
+	{
+		return bSupported;
+	}
+	
+#if PLATFORM_WINDOWS
+	bool bEnabled = false;
+	if (!FParse::Param(FCommandLine::Get(), TEXT("NoUbaController")))
+	{
+		GConfig->GetBool(TEXT("UbaController"), TEXT("Enabled"), bEnabled, GEngineIni);
+	}
+
+	bSupported = FPlatformProcess::SupportsMultithreading() && bEnabled;
+#else
+	// Not supported on other platforms.
+	bSupported = false;
+#endif
+	return bSupported;
+}
+
+void FUbaControllerModule::CleanWorkingDirectory() const
+{
+	IFileManager& FileManager = IFileManager::Get();
+	
+	if (!RootWorkingDirectory.IsEmpty())
+	{
+		if (!FileManager.DeleteDirectory(*RootWorkingDirectory, false, true))
+		{
+			UE_LOG(LogUbaController, Log, TEXT("%s => Failed to delete current working Directory => %s"), ANSI_TO_TCHAR(__FUNCTION__), *RootWorkingDirectory);
+		}	
+	}
+}
+
+bool FUbaControllerModule::HasTasksDispatchedOrPending() const
+{
+	return !PendingRequestedCompilationTasks.IsEmpty() || (JobDispatcherThread.IsValid() && JobDispatcherThread->HasJobsInFlight());
+}
+
+FString GetUbaBinariesPath()
+{
+#if PLATFORM_CPU_ARM_FAMILY
+	const TCHAR* BinariesArch = TEXT("arm64");
+#else
+	const TCHAR* BinariesArch = TEXT("x64");
+#endif
+	return FPaths::Combine(FPaths::EngineDir(), "Binaries/Win64/UnrealBuildAccelerator", BinariesArch);
+}
+
+void FUbaControllerModule::LoadDependencies()
+{
+	const FString UbaBinariesPath = GetUbaBinariesPath();
+	FPlatformProcess::AddDllDirectory(*UbaBinariesPath);
+	FPlatformProcess::GetDllHandle(*(FPaths::Combine(UbaBinariesPath, "UbaHost.dll")));
+}
+
+void FUbaControllerModule::StartupModule()
+{
+	check(!bModuleInitialized);
+
+	LoadDependencies();
+
+	IModularFeatures::Get().RegisterModularFeature(GetModularFeatureType(), this);
+
+	bModuleInitialized = true;
+
+	FCoreDelegates::OnEnginePreExit.AddLambda([&]()
+	{
+		if (bControllerInitialized && JobDispatcherThread)
+		{
+			JobDispatcherThread->Stop();
+			FPlatformProcess::ConditionalSleep([&]() { return JobDispatcherThread && JobDispatcherThread->IsWorkDone(); }, 0.1f);
+			JobDispatcherThread = nullptr;
+		}
+	});
+}
+
+void FUbaControllerModule::ShutdownModule()
+{
+	check(bModuleInitialized);
+	
+	IModularFeatures::Get().UnregisterModularFeature(GetModularFeatureType(), this);
+	
+	if (bControllerInitialized)
+	{
+		// Stop the jobs thread
+		if (JobDispatcherThread)
+		{
+			JobDispatcherThread->Stop();	
+			// Wait until the thread is done
+			FPlatformProcess::ConditionalSleep([&](){ return JobDispatcherThread && JobDispatcherThread->IsWorkDone(); },0.1f);
+		}
+
+		FTask* Task;
+		while (PendingRequestedCompilationTasks.Dequeue(Task))
+		{
+			FDistributedBuildTaskResult Result;
+			Result.ReturnCode = 0;
+			Result.bCompleted = false;
+			Task->Promise.SetValue(Result);
+			delete Task;
+		}
+
+		PendingRequestedCompilationTasks.Empty();
+	}
+
+	CleanWorkingDirectory();
+	bModuleInitialized = false;
+	bControllerInitialized = false;
+}
+
+void FUbaControllerModule::InitializeController()
+{
+	// We should never Initialize the controller twice
+	if (ensureAlwaysMsgf(!bControllerInitialized, TEXT("Multiple initialization of UBA controller!")))
+	{
+		CleanWorkingDirectory();
+
+		if (IsSupported())
+		{
+			IFileManager::Get().MakeDirectory(*WorkingDirectory, true);
+
+			// Pre-create the directories so we don't have to explicitly register them to uba later
+			for (int32 It = 0; It != UbaControllerModule::SubFolderCount; ++It)
+			{
+				IFileManager::Get().MakeDirectory(*FString::Printf(TEXT("%s/%d"), *WorkingDirectory, It));
+			}
+
+
+			JobDispatcherThread = MakeShared<FUbaJobProcessor>(*this);
+			JobDispatcherThread->StartThread();
+		}
+
+		bControllerInitialized = true;
+	}
+}
+
+FString FUbaControllerModule::CreateUniqueFilePath()
+{
+	check(bSupported);
+	int32 ID = NextFileID++;
+	int32 FolderID = ID % UbaControllerModule::SubFolderCount; // We use sub folders to be nicer to file system (we can end up with 20000 files in one folder otherwise)
+	return FString::Printf(TEXT("%s/%d/%d.uba"), *WorkingDirectory, FolderID, ID);
+}
+
+TFuture<FDistributedBuildTaskResult> FUbaControllerModule::EnqueueTask(const FTaskCommandData& CommandData)
+{
+	check(bSupported);
+
+	TPromise<FDistributedBuildTaskResult> Promise;
+	TFuture<FDistributedBuildTaskResult> Future = Promise.GetFuture();
+
+	// Enqueue the new task
+	FTask* Task = new FTask(NextTaskID++, CommandData, MoveTemp(Promise));
+	{
+		PendingRequestedCompilationTasks.Enqueue(Task);
+	}
+
+	JobDispatcherThread->HandleTaskQueueUpdated(CommandData.InputFileName);
+
+	return MoveTemp(Future);
+}
+
+void FUbaControllerModule::ReportJobProcessed(const FTaskResponse& InTaskResponse, FTask* CompileTask)
+{
+	if (CompileTask)
+	{
+		FDistributedBuildTaskResult Result;
+		Result.ReturnCode = InTaskResponse.ReturnCode;
+		Result.bCompleted = true;
+		CompileTask->Promise.SetValue(Result);
+		delete CompileTask;
+	}
+}
+
+UBACONTROLLER_API FUbaControllerModule& FUbaControllerModule::Get()
+{
+	return FModuleManager::LoadModuleChecked<FUbaControllerModule>(TEXT("UbaController"));
+}
+
+IMPLEMENT_MODULE(FUbaControllerModule, UbaController);
diff --git a/Engine/Plugins/UbaController/Source/UbaController/Private/UbaHordeAgent.cpp b/Engine/Plugins/UbaController/Source/UbaController/Private/UbaHordeAgent.cpp
new file mode 100644
--- /dev/null
+++ b/Engine/Plugins/UbaController/Source/UbaController/Private/UbaHordeAgent.cpp
@@ -0,0 +1,237 @@
+// Copyright Epic Games, Inc. All Rights Reserved.
+
+#include "UbaHordeAgent.h"
+
+#include "Misc/Paths.h"
+#include "HAL/FileManager.h"
+#include "Compute/AgentMessage.h"
+#include "Compute/ComputeSocket.h"
+#include "Compute/ComputeTransport.h"
+#include "UbaHordeComputeTransport.h"
+
+DEFINE_LOG_CATEGORY(LogUbaHordeAgent);
+
+FUbaHordeAgent::FUbaHordeAgent(const FHordeRemoteMachineInfo& InMachineInfo)
+	: bIsValid(false)
+	, bHasErrors(false)
+	, MachineInfo(InMachineInfo)
+{
+	// Create the compute transport object and directly send over the nonce.
+	TUniquePtr<FUbaHordeComputeTransport> ComputeTransport = MakeUnique<FUbaHordeComputeTransport>(MachineInfo, bHasErrors);
+	if (!ComputeTransport->IsValid())
+	{
+		return;
+	}
+
+	ComputeTransport->Send(MachineInfo.Nonce, sizeof(MachineInfo.Nonce));
+
+	// Create the compute socket and initialize a compute channel with the recv/send buffers.
+	// This will allow us to create the agent channel object which will allow us to communicate
+	// with the agent directly.
+	HordeComputeSocket = CreateComputeSocket(MoveTemp(ComputeTransport), EComputeSocketEndpoint::Remote);
+	TSharedPtr<FComputeChannel> ComputeChannel = HordeComputeSocket->CreateChannel(0);
+	TSharedPtr<FComputeChannel> SecondComputeChannel = HordeComputeSocket->CreateChannel(100);
+
+	// Now this is what is going to handle the ComputeChannel.
+	AgentChannel = MakeUnique<FAgentMessageChannel>(MoveTemp(ComputeChannel));
+	ChildChannel = MakeUnique<FAgentMessageChannel>(MoveTemp(SecondComputeChannel));
+
+	bIsValid = true;
+}
+
+FUbaHordeAgent::~FUbaHordeAgent()
+{
+}
+
+bool FUbaHordeAgent::BeginCommunication()
+{
+	ensure(bIsValid);
+
+	// Now start the communication. From now, we can send and receive stuff.
+	HordeComputeSocket->StartCommunication();
+
+	// Let's try to read the response from the server. It should be an Attach message.
+	EAgentMessageType Type = AgentChannel->ReadResponse(5000);
+	UE_LOG(LogUbaHordeAgent, Log, TEXT("Got a response from the server: Type=0x%08X"), (int)Type);
+
+	if (Type == EAgentMessageType::None) // Timed out
+	{
+		return false;
+	}
+
+	if (!ensure(Type == EAgentMessageType::Attach))
+	{
+		return false;
+	}
+
+	/* Fork communication channel to split communication between a main and child channels. */
+	AgentChannel->Fork(100, 4*1024*1024);
+
+	Type = ChildChannel->ReadResponse(5000);
+
+	if (Type == EAgentMessageType::None) // Timed out
+	{
+		return false;
+	}
+
+	if (!ensure(Type == EAgentMessageType::Attach))
+	{
+		return false;
+	}
+
+	return true;
+}
+
+bool FUbaHordeAgent::UploadBinaries(const FString& BundleDirectory, const char* BundleLocator)
+{
+	// Prepare for giving the remote machine work.
+	// 1) Send binaries to remote machine.
+	// 2) Request to start executing UbaAgent.exe.
+	// 3) From there, it's a simple matter of sending work over through UBA.
+	ChildChannel->UploadFiles("", BundleLocator);
+
+	TMap<FString, FArchive*> BlobFiles;
+
+	auto FindOrAddBlobFile = [&BlobFiles, &BundleDirectory](const AgentMessage::FBlobRequest& BlobRequest) -> FArchive*
+		{
+			const FString Locator(BlobRequest.Locator.GetData());
+
+			if (Locator.IsEmpty())
+			{
+				UE_LOG(LogUbaHordeAgent, Error, TEXT("Cannot upload binaries to Horde agent with empty locator"));
+				return nullptr;
+			}
+
+			if (FArchive** ArchiveElem = BlobFiles.Find(Locator))
+			{
+				return *ArchiveElem;
+			}
+			else
+			{
+				FString Path = FPaths::Combine(BundleDirectory, BlobRequest.Locator.GetData());
+				Path.Append(TEXT(".blob"));
+
+				if (FArchive* Archive = IFileManager::Get().CreateFileReader(*Path))
+				{
+					BlobFiles.Add({ Locator, Archive });
+					return Archive;
+				}
+				else
+				{
+					UE_LOG(LogUbaHordeAgent, Error, TEXT("Cannot read blob file for Horde agent upload: %s"), *Path);
+					return nullptr;
+				}
+			}
+		};
+
+	EAgentMessageType Type;
+
+	for (;;)
+	{
+		Type = ChildChannel->ReadResponse();
+		if (Type != EAgentMessageType::ReadBlob)
+		{
+			break;
+		}
+
+		AgentMessage::FBlobRequest BlobRequest;
+		ChildChannel->ReadBlobRequest(BlobRequest);
+
+		TArray<UTF8CHAR> LocatorData(BlobRequest.Locator.GetData(), BlobRequest.Locator.Len());
+		LocatorData.Add(UTF8CHAR('\0'));
+		UE_LOG(LogUbaHordeAgent, Log, TEXT("Response [ReadBlob]: Locator=%s, Offset=%d, Length=%d"), UTF8_TO_TCHAR(LocatorData.GetData()), BlobRequest.Offset, BlobRequest.Length);
+
+		FArchive* Archive = FindOrAddBlobFile(BlobRequest);
+		if (Archive == nullptr)
+		{
+			return false;
+		}
+
+		Archive->Seek(BlobRequest.Offset);
+
+		TArray<uint8> SerializedBytes;
+		SerializedBytes.SetNum(Archive->TotalSize() - BlobRequest.Offset);
+		Archive->Serialize(SerializedBytes.GetData(), SerializedBytes.Num());
+
+		// Send the blob data.
+		ChildChannel->Blob((unsigned char*)SerializedBytes.GetData(), SerializedBytes.Num());
+	}
+
+	if (!ensure(Type == EAgentMessageType::WriteFilesResponse))
+	{
+		return false;
+	}
+
+	return true;
+}
+
+void FUbaHordeAgent::Execute(const char* Exe, const char** Args, size_t NumArgs, const char* WorkingDir, const char** EnvVars, size_t NumEnvVars, bool bUseWine)
+{
+	ChildChannel->Execute(Exe, Args, NumArgs, WorkingDir, EnvVars, NumEnvVars, bUseWine ? EExecuteProcessFlags::UseWine : EExecuteProcessFlags::None);
+}
+
+void FUbaHordeAgent::CloseConnection()
+{
+	ChildChannel->Close();
+	AgentChannel->Close();
+}
+
+bool FUbaHordeAgent::IsValid()
+{
+	return bIsValid && !bHasErrors;
+}
+
+void FUbaHordeAgent::PollReports()
+{
+	EAgentMessageType Type = EAgentMessageType::None;
+	constexpr int32 ReadresponseTimeoutMS = 100;
+	while ((Type = ChildChannel->ReadResponse(ReadresponseTimeoutMS)) != EAgentMessageType::None)
+	{
+		switch (Type)
+		{
+		case EAgentMessageType::ExecuteOutput:
+		{
+			if (ChildChannel->GetResponseSize() > 0)
+			{
+				// Convert raw buffer to dynamic array for modification
+				const ANSICHAR* ResponseDataRaw = reinterpret_cast<const ANSICHAR*>(ChildChannel->GetResponseData());
+				TArray<ANSICHAR> ResponseData(ResponseDataRaw, ChildChannel->GetResponseSize());
+
+				// Remove trailing newline characters and add NUL-terminator
+				while (ResponseData.Num() > 0 && (ResponseData.Last() == '\n' || ResponseData.Last() == '\r'))
+				{
+					constexpr bool bAllowShrinking = false;
+					ResponseData.Pop(bAllowShrinking);
+				}
+				ResponseData.Add('\0');
+
+				// Report output to log
+				auto ResponseString = StringCast<TCHAR>(ResponseData.GetData());
+				if (ResponseString.Length() > 0)
+				{
+					UE_LOG(LogUbaHordeAgent, Log, TEXT("Response [ExecuteOutput]: %s"), ResponseString.Get());
+				}
+			}
+		}
+		break;
+
+		case EAgentMessageType::ExecuteResult:
+		{
+			if (ChildChannel->GetResponseSize() == sizeof(int32))
+			{
+				const int32* ResponseData = reinterpret_cast<const int32*>(ChildChannel->GetResponseData());
+				const int32 ExecuteExitCode = ResponseData[0];
+				UE_LOG(LogUbaHorde, Log, TEXT("Response [ExecuteResult]: ExitCode=%d"), ExecuteExitCode);
+				bIsValid = false;
+			}
+		}
+		break;
+
+		default:
+		{
+			// Ignore other messages here, we are only interested in the ExecuteOutput to forward reports from the remote agents to the calling process
+		}
+		break;
+		}
+	}
+}
diff --git a/Engine/Plugins/UbaController/Source/UbaController/Private/UbaHordeAgentManager.cpp b/Engine/Plugins/UbaController/Source/UbaController/Private/UbaHordeAgentManager.cpp
new file mode 100644
--- /dev/null
+++ b/Engine/Plugins/UbaController/Source/UbaController/Private/UbaHordeAgentManager.cpp
@@ -0,0 +1,383 @@
+// Copyright Epic Games, Inc. All Rights Reserved.
+
+#include "UbaHordeAgentManager.h"
+#include "HAL/Event.h"
+#include "HAL/IConsoleManager.h"
+#include "Misc/ConfigCacheIni.h"
+#include "Misc/FileHelper.h"
+#include "Misc/Paths.h"
+#include "Misc/ScopeExit.h"
+#include "Storage/StorageClient.h"
+#include "Storage/Clients/BundleStorageClient.h"
+#include "Storage/Clients/MemoryStorageClient.h"
+#include "Storage/Clients/FileStorageClient.h"
+#include "Storage/Nodes/ChunkNode.h"
+#include "Storage/Nodes/DirectoryNode.h"
+#include "Storage/BlobWriter.h"
+#include "UbaControllerModule.h"
+#include "UbaExports.h"
+#include <filesystem>
+#include <fstream>
+
+namespace UbaControllerModule
+{
+	static bool bHordeForwardAgentLogs = false;
+	static FAutoConsoleVariableRef CVarUbaControllerHordeForwardAgentLogs(
+		TEXT("r.UbaController.HordeForwardAgentLogs"),
+		bHordeForwardAgentLogs,
+		TEXT("Enables or disables the use of UBA as controller for distributed help in the engine.\n")
+		TEXT("false: UBA will not be used as controller \n")
+		TEXT("true: Distribute builds using UBA (default) is no other controller is available."),
+		ECVF_ReadOnly); // Must be set on start-up, e.g. via config ini
+}
+
+FUbaHordeAgentManager::FUbaHordeAgentManager(uba::NetworkServer* InServer)
+	:	UbaServer(InServer)
+	,	LastRequestFailTime(1)
+	,	TargetCoreCount(0)
+	,	EstimatedCoreCount(0)
+	,	AskForAgents(true)
+{
+	ParseConfig();
+}
+
+FUbaHordeAgentManager::~FUbaHordeAgentManager()
+{
+	FScopeLock AgentsScopeLock(&AgentsLock);
+	for (TUniquePtr<FHordeAgentWrapper>& Agent : Agents)
+	{
+		Agent->ShouldExit->Trigger();
+		Agent->Thread.Join();
+		FGenericPlatformProcess::ReturnSynchEventToPool(Agent->ShouldExit);
+	}
+}
+
+void FUbaHordeAgentManager::SetTargetCoreCount(uint32 Count)
+{
+	TargetCoreCount = FMath::Min(MaxCores, Count);
+
+	if (Url.IsEmpty())
+	{
+		AskForAgents = false;
+	}
+
+	while (EstimatedCoreCount < TargetCoreCount)
+	{
+		if (!AskForAgents)
+		{
+			return;
+		}
+
+		//UE_LOG(LogUbaController, Display, TEXT("Requested new agent. Estimated core count: %u, Target core count: %u"), EstimatedCoreCount.Load(), TargetCoreCount.Load());
+		RequestAgent();
+	}
+
+	FScopeLock AgentsScopeLock(&AgentsLock);
+	for (auto Iterator = Agents.CreateIterator(); Iterator; ++Iterator)
+	{
+		TUniquePtr<FHordeAgentWrapper>& Agent = *Iterator;
+		if (Agent->ShouldExit->Wait(0))
+		{
+			Agent->Thread.Join();
+			FGenericPlatformProcess::ReturnSynchEventToPool(Agent->ShouldExit);
+			Iterator.RemoveCurrentSwap();
+		}
+	}
+}
+
+// Creates a bundle blob (one of several chunks of a file) to be uploaded to Horde
+// This code has been adopted from the HordeTest project.
+// See 'Engine/Source/Programs/Horde/Samples/HordeTest/Main.cpp'.
+static FBlobHandleWithHash CreateHordeBundleBlob(std::ifstream& Stream, FBlobWriter& Writer, int64& OutLength, FIoHash& OutStreamHash)
+{
+	OutLength = 0;
+
+	FChunkNodeWriter ChunkWriter(Writer);
+
+	char ReadBuffer[4096];
+	while (!Stream.eof())
+	{
+		Stream.read(ReadBuffer, sizeof(ReadBuffer));
+
+		const int64 ReadSize = static_cast<int64>(Stream.gcount());
+		if (ReadSize == 0)
+		{
+			break;
+		}
+		OutLength += ReadSize;
+
+		ChunkWriter.Write(FMemoryView(ReadBuffer, ReadSize));
+	}
+
+	return ChunkWriter.Flush(OutStreamHash);
+}
+
+static FFileEntry CreateHordeBundleFileEntry(const std::filesystem::path& Path, FBlobWriter& Writer)
+{
+	std::ifstream InputStream(Path, std::ios::binary);
+	check(InputStream.good());
+
+	int64 Length = 0;
+	FIoHash StreamHash;
+	FBlobHandleWithHash Target = CreateHordeBundleBlob(InputStream, Writer, Length, StreamHash);
+
+	return FFileEntry(Target, FUtf8String(Path.filename().string().c_str()), EFileEntryFlags::None, Length, StreamHash, FSharedBufferView());
+}
+
+static FDirectoryEntry CreateHordeBundleDirectoryEntry(const std::filesystem::path& Path, FBlobWriter& Writer)
+{
+	FDirectoryNode DirectoryNode;
+
+	FFileEntry NewEntry = CreateHordeBundleFileEntry(Path, Writer);
+	const FUtf8String Name = NewEntry.Name;
+	const int64 Length = NewEntry.Length;
+	DirectoryNode.NameToFile.Add(Name, MoveTemp(NewEntry));
+
+	FBlobHandle DirectoryHandle = DirectoryNode.Write(Writer);
+
+	return FDirectoryEntry(DirectoryHandle, FIoHash(), FUtf8String(Path.filename().string().c_str()), Length);
+}
+
+bool CreateHordeBundleFromFile(const std::filesystem::path& InputFilename, const std::filesystem::path& OutputFilename)
+{
+	TSharedRef<FFileStorageClient> FileStorage = MakeShared<FFileStorageClient>(OutputFilename.parent_path());
+	TSharedRef<FBundleStorageClient> Storage = MakeShared<FBundleStorageClient>(FileStorage);
+
+	TUniquePtr<FBlobWriter> Writer = Storage->CreateWriter("");
+	FDirectoryEntry RootEntry = CreateHordeBundleDirectoryEntry(InputFilename, *Writer.Get());
+	Writer->Flush();
+
+	FFileStorageClient::WriteRefToFile(OutputFilename, RootEntry.Target->GetLocator());
+	return true;
+}
+
+FString GetUbaBinariesPath();
+
+void FUbaHordeAgentManager::RequestAgent()
+{
+	EstimatedCoreCount += 32; // We estimate a typical agent to have 32 cores
+
+	FScopeLock AgentsScopeLock(&AgentsLock);
+	FHordeAgentWrapper& Wrapper = *Agents.Emplace_GetRef(MakeUnique<FHordeAgentWrapper>());
+
+	Wrapper.ShouldExit = FGenericPlatformProcess::GetSynchEventFromPool(true);
+	Wrapper.Thread = FThread(
+		TEXT("HordeAgent"),
+		[this, WrapperPtr = &Wrapper]()
+		{
+			FEvent& ShouldExit = *WrapperPtr->ShouldExit;
+			TUniquePtr<FUbaHordeAgent> Agent;
+			bool bSuccess = false;
+
+			ON_SCOPE_EXIT
+			{
+				if (Agent)
+				{
+					Agent->CloseConnection();
+				}
+
+				ShouldExit.Trigger();
+			};
+
+			int MachineCoreCount = 0;
+
+			{
+				ON_SCOPE_EXIT{ EstimatedCoreCount -= 32; };
+
+				FScopeLock ScopeLock(&UbaAgentBundleFilePathLock);
+				if (UbaAgentBundleFilePath.IsEmpty())
+				{
+					FString RootWorkingDirectory = FPaths::Combine(FPlatformProcess::UserTempDir(), TEXT("UbaWorkingDir"));
+
+					const FString UbaAgentFilePath = FPaths::Combine(GetUbaBinariesPath(), TEXT("UbaAgent.exe"));
+					UbaAgentBundleFilePath = FPaths::Combine(RootWorkingDirectory, TEXT("UbaAgent.Bundle.ref"));
+
+					if (!CreateHordeBundleFromFile(*UbaAgentFilePath, *UbaAgentBundleFilePath))
+					{
+						UE_LOG(LogUbaController, Error, TEXT("Failed to create Horde bundle for UbaAgent executable: %s"), *UbaAgentFilePath);
+						AskForAgents = false;
+						return;
+					}
+					UE_LOG(LogUbaController, Display, TEXT("Created Horde bundle for UbaAgent executable: %s"), *UbaAgentFilePath);
+				}
+
+				if (!HordeMetaClient)
+				{
+					// Create Horde meta client right before we need it to make sure the CVar for the server URL has been read by now
+					HordeMetaClient = MakeUnique<FUbaHordeMetaClient>(Url, Oidc);
+					if (!HordeMetaClient->RefreshHttpClient())
+					{
+						UE_LOG(LogUbaController, Error, TEXT("Failed to create HttpClient for UbaAgent"));
+						AskForAgents = false;
+						return;
+					}
+				}
+
+				if (!AskForAgents)
+				{
+					return;
+				}
+
+				if (LastRequestFailTime == 0)
+				{
+					ScopeLock.Unlock();
+				}
+				else
+				{
+					// Try to reduce pressure on horde by not asking for machines more frequent than every 5 seconds if failed to retrieve last time
+					uint64 CurrentTime = FPlatformTime::Cycles64();
+					uint32 MsSinceLastFail = uint32((CurrentTime - LastRequestFailTime) * FPlatformTime::GetSecondsPerCycle() * 1000);
+					if (MsSinceLastFail < 5000)
+					{
+						if (ShouldExit.Wait(5000 - MsSinceLastFail))
+						{
+							return;
+						}
+					}
+				}
+
+				TSharedPtr<FUbaHordeMetaClient::HordeMachinePromise, ESPMode::ThreadSafe> Promise = HordeMetaClient->RequestMachine(Pool);
+				if (!Promise)
+				{
+					//UE_LOG(LogUbaController, Error, TEXT("Failed to create Horde bundle for UbaAgent executable: %s"), *UbaAgentFilePath);
+					return;
+				}
+				TFuture<TTuple<FHttpResponsePtr, FHordeRemoteMachineInfo>> Future = Promise->GetFuture();
+				Future.Wait();
+				FHordeRemoteMachineInfo MachineInfo = Future.Get().Value;
+
+				// If the machine couldn't be assigned, just ignore this agent slot
+				if (MachineInfo.Ip == TEXT(""))
+				{
+					LastRequestFailTime = FPlatformTime::Cycles64();
+					return;
+				}
+
+				LastRequestFailTime = 0;
+
+				ScopeLock.Unlock();
+
+				if (ShouldExit.Wait(0))
+				{
+					return;
+				}
+
+				Agent = MakeUnique<FUbaHordeAgent>(MachineInfo);
+
+				if (!Agent->BeginCommunication())
+				{
+					return;
+				}
+
+				TArray<uint8> Locator;
+				if (!FFileHelper::LoadFileToArray(Locator, *UbaAgentBundleFilePath))
+				{
+					UE_LOG(LogUbaController, Error, TEXT("Cannot launch Horde processes for UBA controller because bundle path could not be found: %s"), *UbaAgentBundleFilePath);
+					return;
+				}
+
+				Locator.Add('\0');
+
+				FString BundleDirectory = FPaths::GetPath(UbaAgentBundleFilePath);
+
+				if (ShouldExit.Wait(0))
+				{
+					return;
+				}
+
+				if (!Agent->UploadBinaries(BundleDirectory, reinterpret_cast<const char*>(Locator.GetData())))
+				{
+					return;
+				}
+
+				uint32 ListenPort = 7001;
+
+				// Start the UBA Agent that will connect to us, requesting for work
+				const std::string ListenPortArg = "-listen=" + std::to_string(ListenPort);
+
+				const char* UbaAgentArgs[] =
+				{
+					ListenPortArg.c_str(),
+					"-nopoll",				// -nopoll recommended when running on remote Horde agents to make sure they exit after completion. Otherwise, it keeps running.
+					"-listenTimeout=20",	// increase timeout for the agent to listen from 5 to 20 seconds to give the editor enough time to establish the connection
+				};
+
+				// If the machine does not run Windows, enable the compatibility layer Wine to run UbaAgent.exe on POSIX systems
+				const bool bRunsWindowsOS = Agent->GetMachineInfo().bRunsWindowOS;
+				const bool bUseWine = !bRunsWindowsOS;
+
+				if (ShouldExit.Wait(0))
+				{
+					return;
+				}
+
+				Agent->Execute("UbaAgent.exe", UbaAgentArgs, UE_ARRAY_COUNT(UbaAgentArgs), nullptr, nullptr, 0, bUseWine);
+
+				// Add this machine as client to the remote agent
+				const FString& IpAddress = Agent->GetMachineInfo().Ip;
+				auto IpAddressStr = StringCast<uba::tchar>(*IpAddress);
+				const bool bAddClientSuccess = Server_AddClient(UbaServer, IpAddressStr.Get(), ListenPort, nullptr);
+
+				if (!bAddClientSuccess)
+				{
+					UE_LOG(LogUbaController, Error, TEXT("Server_AddClient(%s:%d) failed"), *IpAddress, ListenPort);
+					return;
+				}
+
+				// Log remote execution
+				FString UbaAgentCmdArgs = TEXT("UbaAgent.exe");
+				for (const char* Arg : UbaAgentArgs)
+				{
+					UbaAgentCmdArgs += TEXT(" ");
+					UbaAgentCmdArgs += ANSI_TO_TCHAR(Arg);
+				}
+				UE_LOG(LogUbaController, Log, TEXT("Remote execution on Horde machine [%s:%d]: %s"), *IpAddress, ListenPort, *UbaAgentCmdArgs);
+
+				MachineCoreCount = MachineInfo.LogicalCores;
+				EstimatedCoreCount += MachineCoreCount;
+			}
+
+			while (!ShouldExit.Wait(100))
+			{
+				if (!Agent->IsValid())
+					break;
+				if (UbaControllerModule::bHordeForwardAgentLogs)
+					Agent->PollReports();
+			}
+
+			EstimatedCoreCount -= MachineCoreCount;
+		});
+}
+
+void FUbaHordeAgentManager::ParseConfig()
+{
+	// Try to read authentication provider identifier.
+	FString HordeConfig;
+	if (GConfig->GetString(TEXT("UbaController"), TEXT("Horde"), HordeConfig, GEngineIni))
+	{
+		HordeConfig.TrimStartInline();
+		HordeConfig.TrimEndInline();
+		HordeConfig.RemoveFromStart(TEXT("("));
+		HordeConfig.RemoveFromEnd(TEXT(")"));
+
+		if (FParse::Value(*HordeConfig, TEXT("Url="), Url))
+		{
+			UE_LOG(LogUbaHorde, Log, TEXT("Found UBA controller Url: \"%s\""), *Url);
+		}
+
+		if (FParse::Value(*HordeConfig, TEXT("Pool="), Pool))
+		{
+			UE_LOG(LogUbaHorde, Log, TEXT("Found UBA controller Pool: \"%s\""), *Pool);
+		}
+
+		if (FParse::Value(*HordeConfig, TEXT("Oidc="), Oidc))
+		{
+			UE_LOG(LogUbaHorde, Log, TEXT("Found UBA controller Oidc: \"%s\""), *Oidc);
+		}
+
+		if (FParse::Value(*HordeConfig, TEXT("MaxCores="), MaxCores))
+		{
+			UE_LOG(LogUbaHorde, Log, TEXT("Found UBA controller MaxCores: \"%u\""), MaxCores);
+		}
+	}
+}
diff --git a/Engine/Plugins/UbaController/Source/UbaController/Private/UbaHordeComputeTransport.cpp b/Engine/Plugins/UbaController/Source/UbaController/Private/UbaHordeComputeTransport.cpp
new file mode 100644
--- /dev/null
+++ b/Engine/Plugins/UbaController/Source/UbaController/Private/UbaHordeComputeTransport.cpp
@@ -0,0 +1,92 @@
+// Copyright Epic Games, Inc. All Rights Reserved.
+
+#include "UbaHordeComputeTransport.h"
+#include "Sockets.h"
+#include "SocketSubsystem.h"
+
+FUbaHordeComputeTransport::FUbaHordeComputeTransport(const FHordeRemoteMachineInfo& MachineInfo, bool& InHasErrors)
+	: Socket(nullptr)
+	, bIsClosed(false)
+	, bHasErrors(InHasErrors)
+{
+	ISocketSubsystem* SocketSubsystem = ISocketSubsystem::Get(PLATFORM_SOCKETSUBSYSTEM);
+
+	Socket = SocketSubsystem->CreateSocket(TEXT("Stream"), TEXT(""));
+	TSharedPtr<FInternetAddr> InternetAddr;
+
+	TSharedPtr<FInternetAddr> Address = SocketSubsystem->GetAddressFromString(MachineInfo.Ip);
+	Address->SetPort(MachineInfo.Port);
+
+	if (!Socket->Connect(*Address))
+	{
+		UE_LOG(LogUbaHorde, Warning, TEXT("Failed to connect to Horde Agent"));
+		ensure(false);
+
+		SocketSubsystem->DestroySocket(Socket);
+
+		Socket = nullptr;
+	}
+}
+
+FUbaHordeComputeTransport::~FUbaHordeComputeTransport()
+{
+	if (Socket)
+	{
+		ISocketSubsystem* SocketSubsystem = ISocketSubsystem::Get(PLATFORM_SOCKETSUBSYSTEM);
+		SocketSubsystem->DestroySocket(Socket);
+	}
+}
+
+// Sends data to the remote
+size_t FUbaHordeComputeTransport::Send(const void* Data, size_t Size)
+{
+	int32 NumBytesSent = 0;
+	if (!Socket->Send((const uint8*)Data, Size, NumBytesSent))
+	{
+		NumBytesSent = 0;
+		// We can't log this since the other side could have disconnected us (causing recv to fail). This happens often in the horde setup
+		//UE_LOG(LogUbaHorde, Error, TEXT("Failed to send data to the Horde Agent: %llu %s"), (uint64)Size, Size == 1 ? TEXT("byte") : TEXT("bytes"));
+		bHasErrors = true;
+		return 0;
+	}
+
+	UE_LOG(LogUbaHorde, Verbose, TEXT("Sent message to horde agent: %d %s"), NumBytesSent, NumBytesSent == 1 ? TEXT("byte") : TEXT("bytes"));
+
+	return NumBytesSent;
+}
+
+// Receives data from the remote
+size_t FUbaHordeComputeTransport::Recv(void* Data, size_t Size)
+{
+	int32 NumBytesRead = 0;
+	if (!Socket->Recv((uint8*)Data, Size, NumBytesRead))
+	{
+		NumBytesRead = 0;
+
+		// We can't log this since the other side could have disconnected us (causing recv to fail). This happens often in the horde setup
+		//if (!bIsClosed)
+		//	UE_LOG(LogUbaHorde, Error, TEXT("Failed to receive data from the Horde Agent: %llu %s"), (uint64)Size, Size == 1 ? TEXT("byte") : TEXT("bytes"));
+		return 0;
+	}
+
+	UE_LOG(LogUbaHorde, Verbose, TEXT("Received message from horde agent: %d %s"), NumBytesRead, NumBytesRead == 1 ? TEXT("byte") : TEXT("bytes"));
+
+	return NumBytesRead;
+}
+
+// Indicates to the remote that no more data will be sent.
+void FUbaHordeComputeTransport::MarkComplete()
+{
+}
+
+// Indicates that no more data will be sent or received, and that any blocking reads/writes should stop.
+void FUbaHordeComputeTransport::Close()
+{
+	bIsClosed = true;
+	Socket->Close();
+}
+
+bool FUbaHordeComputeTransport::IsValid()
+{
+	return Socket != nullptr && !bHasErrors;
+}
diff --git a/Engine/Plugins/UbaController/Source/UbaController/Private/UbaHordeMetaClient.cpp b/Engine/Plugins/UbaController/Source/UbaController/Private/UbaHordeMetaClient.cpp
new file mode 100644
--- /dev/null
+++ b/Engine/Plugins/UbaController/Source/UbaController/Private/UbaHordeMetaClient.cpp
@@ -0,0 +1,152 @@
+// Copyright Epic Games, Inc. All Rights Reserved.
+
+#include "UbaHordeMetaClient.h"
+#include "Serialization/JsonReader.h"
+#include "Serialization/JsonSerializer.h"
+#include "Misc/App.h"
+#include "Misc/ConfigCacheIni.h"
+
+DEFINE_LOG_CATEGORY(LogUbaHorde);
+
+FUbaHordeMetaClient::FUbaHordeMetaClient(const FStringView& HordeServerUrl, const FStringView& InOAuthProviderIdentifier)
+	: ServerUrl(HordeServerUrl)
+	, bConnectWithAuthentication(HordeServerUrl.StartsWith(TEXT("https:")))
+	, OAuthProviderIdentifier(InOAuthProviderIdentifier)
+{
+}
+
+FUbaHordeMetaClient::FUbaHordeMetaClient(const FStringView& HordeServerIp, uint16 HordeServerPort, bool bInConnectWithAuthentication, const FStringView& InOAuthProviderIdentifier)
+	: ServerUrl(FString::Format(TEXT("{0}://{1}:{2}"), {(bInConnectWithAuthentication ? TEXT("https") : TEXT("http")), HordeServerIp, HordeServerPort}))
+	, bConnectWithAuthentication(bInConnectWithAuthentication)
+	, OAuthProviderIdentifier(InOAuthProviderIdentifier)
+{
+}
+
+bool FUbaHordeMetaClient::RefreshHttpClient()
+{
+	// Try to connect to Horde with HTTP and v2 API
+	HttpClient = MakeUnique<FHordeHttpClient>(ServerUrl);
+
+	if (bConnectWithAuthentication)
+	{
+		check(!OAuthProviderIdentifier.IsEmpty());
+		UE_LOG(LogUbaHorde, Display, TEXT("Logging in to Horde server with OIDC: %s"), *ServerUrl);
+		if (!HttpClient->LoginWithOidc(*OAuthProviderIdentifier, FApp::IsUnattended()))
+		{
+			UE_LOG(LogUbaHorde, Error, TEXT("Login to Horde server [%s] failed"), *ServerUrl);
+			return false;
+		}
+	}
+	return true;
+}
+
+TSharedPtr<FUbaHordeMetaClient::HordeMachinePromise, ESPMode::ThreadSafe> FUbaHordeMetaClient::RequestMachine(const FString& PoolId, const FString& Machine)
+{
+	TSharedPtr<HordeMachinePromise> Promise = MakeShared<HordeMachinePromise, ESPMode::ThreadSafe>();
+
+	const FString ResourcePath = FString::Format(TEXT("api/v2/compute/{0}"), { Machine });
+
+	FHttpRequestRef Request = HttpClient->CreateRequest(TEXT("POST"), *ResourcePath);
+
+	Request->SetDelegateThreadPolicy(EHttpRequestDelegateThreadPolicy::CompleteOnHttpThread);
+
+	// Post JSON document with constraints to acquire a Horde agent. Use pool of agents (e.g. "BoxLinux" or "BoxWin")
+	// and require exclusive access or UbaStorage will fail to initialize the next time we connect to the same machine.
+	const FString Body = FString::Format(TEXT("{\"requirements\":{\"pool\":\"{0}\",\"exclusive\":true}}"), { PoolId });
+	Request->SetContentAsString(Body);
+	Request->SetHeader(TEXT("Content-Type"), TEXT("application/json"));
+
+	UE_LOG(LogUbaHorde, Log, TEXT("Requesting Horde agent with JSON descriptor: '%s'"), *Body);
+
+	Request->OnProcessRequestComplete().BindLambda(
+		[this, Promise](FHttpRequestPtr /*Request*/, FHttpResponsePtr HttpResponse, bool bSucceeded)
+		{
+			FHordeRemoteMachineInfo Info;
+			Info.Ip = TEXT("");
+			Info.Port = 0xFFFF;
+			Info.bRunsWindowOS = false;
+			FMemory::Memset(Info.Nonce, 0, sizeof(Info.Nonce));
+
+			if (!bSucceeded || !HttpResponse.IsValid())
+			{
+				UE_LOG(LogUbaHorde, Warning, TEXT("No response from Horde"));
+				Promise->SetValue(MakeTuple(HttpResponse, Info));
+				return;
+			}
+
+			FString ResponseStr = HttpResponse->GetContentAsString();
+
+			if (HttpResponse->GetResponseCode() == 503)
+			{
+				UE_LOG(LogUbaHorde, Display, TEXT("No resources available in Horde (%s)"), *ResponseStr);
+				Promise->SetValue(MakeTuple(HttpResponse, Info));
+				return;
+			}
+
+			TSharedPtr<FJsonValue> OutJson;
+			TSharedRef<TJsonReader<>> Reader = TJsonReaderFactory<>::Create(ResponseStr);
+
+			if (!FJsonSerializer::Deserialize(Reader, OutJson, FJsonSerializer::EFlags::None))
+			{
+				UE_LOG(LogUbaHorde, Error, TEXT("Invalid response body: %s"), *ResponseStr);
+				Promise->SetValue(MakeTuple(HttpResponse, Info));
+				return;
+			}
+
+			TSharedPtr<FJsonValue> NonceValue = OutJson->AsObject()->TryGetField(TEXT("nonce"));
+			TSharedPtr<FJsonValue> IpValue = OutJson->AsObject()->TryGetField(TEXT("ip"));
+			TSharedPtr<FJsonValue> PortValue = OutJson->AsObject()->TryGetField(TEXT("port"));
+
+			if (!NonceValue.Get() || !IpValue.Get() || !PortValue.Get())
+			{
+				UE_LOG(LogUbaHorde, Error, TEXT("Invalid response body: %s"), *ResponseStr);
+				Promise->SetValue(MakeTuple(HttpResponse, Info));
+				return;
+			}
+
+			// Search for property "OSFamily=Windows". If not found, assume we are running on a POSIX system and need Wine to run UBA
+			if (TSharedPtr<FJsonValue> PropertiesValue = OutJson->AsObject()->TryGetField(TEXT("properties")))
+			{
+				for (const TSharedPtr<FJsonValue>& PropertyEntryValue : PropertiesValue->AsArray())
+				{
+					checkf(PropertyEntryValue.Get(), TEXT("null pointer in JSON array object of node \"properties\""));
+					const FString PropertyElementString = PropertyEntryValue->AsString();
+					if (PropertyElementString == TEXT("OSFamily=Windows"))
+					{
+						Info.bRunsWindowOS = true;
+					}
+					if (PropertyElementString.StartsWith(TEXT("LogicalCores=")))
+					{
+						Info.LogicalCores = FCString::Atoi(*PropertyElementString + 13);
+					}
+				}
+			}
+
+			FString NonceString = NonceValue->AsString(); 
+			FString IpString = IpValue->AsString();
+			uint16 PortNumber = (uint16)PortValue->AsNumber();
+
+			if (TSharedPtr<FJsonValue> AgentIdValue = OutJson->AsObject()->TryGetField(TEXT("agentId")))
+			{
+				const FString AgentWebPortalUrl = FString::Format(TEXT("{0}/agents?agentId={1}"), { this->ServerUrl, AgentIdValue->AsString() });
+				UE_LOG(
+					LogUbaHorde, Display, TEXT("UBA Horde machine assigned (%s) [%s:%u]: %s"),
+					(Info.bRunsWindowOS ? TEXT("Microsoft Windows") : TEXT("POSIX/Wine")), *IpString, (uint32)PortNumber, *AgentWebPortalUrl
+				);
+			}
+			else
+			{
+				UE_LOG(LogUbaHorde, Display, TEXT("UBA Horde machine assigned [%s:%u]"), *IpString, (uint32)PortNumber);
+			}
+
+			Info.Ip = IpString;
+			Info.Port = PortNumber;
+			FString::ToHexBlob(NonceString, Info.Nonce, HORDE_NONCE_SIZE);
+
+			Promise->SetValue(MakeTuple(HttpResponse, Info));
+		});
+
+	Request->ProcessRequest();
+
+	return Promise;
+}
diff --git a/Engine/Plugins/UbaController/Source/UbaController/Private/UbaJobProcessor.cpp b/Engine/Plugins/UbaController/Source/UbaController/Private/UbaJobProcessor.cpp
new file mode 100644
--- /dev/null
+++ b/Engine/Plugins/UbaController/Source/UbaController/Private/UbaJobProcessor.cpp
@@ -0,0 +1,459 @@
+// Copyright Epic Games, Inc. All Rights Reserved.
+
+#include "UbaJobProcessor.h"
+
+#include "Interfaces/ITargetPlatform.h"
+#include "Interfaces/ITargetPlatformManagerModule.h"
+#include "Misc/ConfigCacheIni.h"
+#include "UbaHordeAgentManager.h"
+#include "Windows/AllowWindowsPlatformTypes.h"
+#include "UbaControllerModule.h"
+#include "Windows/HideWindowsPlatformTypes.h"
+
+#include "HAL/FileManager.h"
+#include "HAL/PlatformFileManager.h"
+
+#include "Misc/Paths.h"
+#include "Misc/ScopeRWLock.h"
+
+namespace UbaJobProcessorOptions
+{
+	static float SleepTimeBetweenActions = 0.01f;
+	static FAutoConsoleVariableRef CVarSleepTimeBetweenActions(
+        TEXT("r.UbaController.SleepTimeBetweenActions"),
+        SleepTimeBetweenActions,
+        TEXT("How much time the job processor thread should sleep between actions .\n"));
+
+	static float MaxTimeWithoutTasks = 10.0f;
+	static FAutoConsoleVariableRef CVarMaxTimeWithoutTasks(
+        TEXT("r.UbaController.MaxTimeWithoutTasks"),
+        MaxTimeWithoutTasks,
+        TEXT("Time to wait (in seconds) before stop processing attempts if we don't have any pending task.\n"));
+
+	static bool bAutoLaunchVisualizer = false;
+	static FAutoConsoleVariableRef CVarAutoLaunchVisualizer(
+		TEXT("r.UbaController.AutoLaunchVisualizer"),
+		bAutoLaunchVisualizer,
+		TEXT("If true, UBA visualizer will be launched automatically\n"));
+
+	static FString TraceFilename;
+	static FAutoConsoleVariableRef CVarTraceFilename(
+		TEXT("r.UbaController.TraceFilename"),
+		TraceFilename,
+		TEXT("The name of the trace file that uba outputs after a session"));
+
+	static bool bAllowProcessReuse = true;
+	static FAutoConsoleVariableRef CVarAllowProcessReuse(
+		TEXT("r.UbaController.AllowProcessReuse"),
+		bAllowProcessReuse,
+		TEXT("If true, remote process is allowed to fetch new processes from the queue (this requires the remote processes to have UbaRequestNextProcess implemented)\n"));
+
+	static bool bDetailedTrace = false;
+	static FAutoConsoleVariableRef CVarDetailedTrace(
+		TEXT("r.UbaController.DetailedTrace"),
+		bDetailedTrace,
+		TEXT("If true, a UBA will output detailed trace\n"));
+
+	static bool bShowUbaLog = false;
+	static FAutoConsoleVariableRef CVarShowUbaLog(
+		TEXT("r.UbaController.ShowUbaLog"),
+		bShowUbaLog,
+		TEXT("If true, UBA log entries will be visible in the log\n"));
+
+	FString ReplaceEnvironmentVariablesInPath(const FString& ExtraFilePartialPath) // Duplicated code with FAST build.. put it somewhere else?
+	{
+		FString ParsedPath;
+
+		// Fast build cannot read environmental variables easily
+		// Is better to resolve them here
+		if (ExtraFilePartialPath.Contains(TEXT("%")))
+		{
+			TArray<FString> PathSections;
+			ExtraFilePartialPath.ParseIntoArray(PathSections, TEXT("/"));
+
+			for (FString& Section : PathSections)
+			{
+				if (Section.Contains(TEXT("%")))
+				{
+					Section.RemoveFromStart(TEXT("%"));
+					Section.RemoveFromEnd(TEXT("%"));
+					Section = FPlatformMisc::GetEnvironmentVariable(*Section);
+				}
+			}
+
+			for (FString& Section : PathSections)
+			{
+				ParsedPath /= Section;
+			}
+
+			FPaths::NormalizeDirectoryName(ParsedPath);
+		}
+
+		if (ParsedPath.IsEmpty())
+		{
+			ParsedPath = ExtraFilePartialPath;
+		}
+
+		return ParsedPath;
+	}
+}
+
+FUbaJobProcessor::FUbaJobProcessor(
+	FUbaControllerModule& InControllerModule) :
+
+	Thread(nullptr),
+	ControllerModule(InControllerModule),
+	bForceStop(false),
+	LastTimeCheckedForTasks(0),
+	bShouldProcessJobs(false),
+	bIsWorkDone(false),
+	LogWriter([]() {}, []() {}, [](uba::LogEntryType type, const wchar_t* str, uba::u32 strlen)
+	{
+			if (UbaJobProcessorOptions::bShowUbaLog)
+				UE_LOG(LogUbaController, Log, TEXT("%s"), str);
+	})
+{
+	Uba_SetCustomAssertHandler([](const uba::tchar* text)
+		{
+			checkf(false, TEXT("%s"), text);
+		});
+
+	if (!GConfig->GetInt(TEXT("UbaController"), TEXT("MaxLocalParallelJobs"), MaxLocalParallelJobs, GEngineIni))
+	{
+		MaxLocalParallelJobs = -1;
+	}
+
+	if (MaxLocalParallelJobs == -1)
+	{
+		MaxLocalParallelJobs = FPlatformMisc::NumberOfCoresIncludingHyperthreads();
+	}
+}
+
+FUbaJobProcessor::~FUbaJobProcessor()
+{
+	delete Thread;
+}
+
+void FUbaJobProcessor::CalculateKnownInputs()
+{
+	// TODO: This is ShaderCompileWorker specific and this code is designed to handle all kinds of distributed workload.
+	// Instead this information should be provided from the outside
+
+
+	if (KnownInputsCount) // In order to improve startup we provide some of the input we know will be loaded by ShaderCompileWorker.exe
+	{
+		return;
+	}
+
+	auto AddKnownInput = [&](const FString& file)
+		{
+			auto& fileData = file.GetCharArray();
+			auto num = KnownInputsBuffer.Num();
+			KnownInputsBuffer.SetNum(num + fileData.Num());
+			memcpy(KnownInputsBuffer.GetData() + num, fileData.GetData(), fileData.Num() * sizeof(uba::tchar));
+			++KnownInputsCount;
+		};
+
+	// Get the binaries
+	TArray<FString> KnownFileNames;
+	FString BinDir = FPaths::Combine(FPaths::EngineDir(), TEXT("Binaries/Win64")); // TODO: Need to support other targets than Win64
+	IFileManager::Get().FindFilesRecursive(KnownFileNames, *BinDir, TEXT("ShaderCompileWorker*.*"), true, false);
+	for (const FString& file : KnownFileNames)
+	{
+		if (!file.EndsWith(TEXT(".pdb")))
+		{
+			AddKnownInput(file);
+		}
+	}
+
+	// Get the compiler dependencies for all platforms)
+	ITargetPlatformManagerModule* TargetPlatformManager = GetTargetPlatformManager();
+	for (ITargetPlatform* TargetPlatform : GetTargetPlatformManager()->GetTargetPlatforms())
+	{
+		KnownFileNames.Empty();
+		TargetPlatform->GetShaderCompilerDependencies(KnownFileNames);
+
+		for (const FString& ExtraFilePartialPath : KnownFileNames)
+		{
+			AddKnownInput(UbaJobProcessorOptions::ReplaceEnvironmentVariablesInPath(ExtraFilePartialPath));
+		}
+	}
+
+	// Get all the config files
+	for (const FString& ConfigDir : FPaths::GetExtensionDirs(FPaths::EngineDir(), TEXT("Config")))
+	{
+		KnownFileNames.Empty();
+		IFileManager::Get().FindFilesRecursive(KnownFileNames, *ConfigDir, TEXT("*.ini"), true, false);
+		for (const FString& file : KnownFileNames)
+		{
+			AddKnownInput(file);
+		}
+	}
+
+	KnownInputsBuffer.Add(0);
+}
+
+void FUbaJobProcessor::RunTaskWithUba(FTask* Task)
+{
+	FTaskCommandData& Data = Task->CommandData;
+	SessionServer_RegisterNewFile(UbaSessionServer, *Data.InputFileName);
+
+	FString InputFileName = FPaths::GetCleanFilename(Data.InputFileName);
+	FString OutputFileName = FPaths::GetCleanFilename(Data.OutputFileName);
+	FString Parameters = FString::Printf(TEXT("\"%s/\" %d 0 \"%s\" \"%s\" %s "), *Data.WorkingDirectory, Data.DispatcherPID, *InputFileName, *OutputFileName, *Data.ExtraCommandArgs);
+	FString AppDir = FPaths::GetPath(Data.Command);
+	FString LogPath = FString::FromInt(Task->ID) + TEXT(".Log");
+
+	uba::ProcessStartInfo ProcessInfo;
+	ProcessInfo.application = *Data.Command;
+	ProcessInfo.arguments = *Parameters;
+	ProcessInfo.description = *InputFileName;
+	ProcessInfo.workingDir = *AppDir;
+
+	ProcessInfo.logFile = *LogPath;
+	ProcessInfo.trackInputs = false;
+	
+	struct ExitedInfo
+	{
+		FUbaJobProcessor* Processor;
+		FString InputFile;
+		FTask* Task;
+	};
+
+	auto Info = new ExitedInfo;
+	Info->Processor = this;
+	Info->InputFile = Data.InputFileName;
+	Info->Task = Task;
+
+	ProcessInfo.userData = Info;
+	ProcessInfo.exitedFunc = [](void* userData, const uba::ProcessHandle&)
+		{
+			if (auto Info = (ExitedInfo*)userData) // It can be null if custom message has already handled all of them
+			{
+				IFileManager::Get().Delete(*Info->InputFile);
+				SessionServer_RegisterDeleteFile(Info->Processor->UbaSessionServer, *Info->InputFile);
+				Info->Processor->HandleUbaJobFinished(Info->Task);
+				delete Info;
+			}
+		};
+
+
+	Scheduler_EnqueueProcess(UbaScheduler, ProcessInfo, 1.0f, KnownInputsBuffer.GetData(), KnownInputsBuffer.Num()*sizeof(uba::tchar), KnownInputsCount);
+}
+
+void FUbaJobProcessor::StartUba()
+{
+	checkf(UbaServer == nullptr, TEXT("FUbaJobProcessor::StartUba() was called twice before FUbaJobProcessor::ShutDownUba()"));
+
+	UbaServer = CreateServer(LogWriter);
+
+	FString RootDir = FPaths::Combine(FPlatformProcess::UserTempDir(), TEXT("UnrealUbaTemp"));
+
+	UbaStorageServer = CreateStorageServer(*UbaServer, *RootDir, 0, true, LogWriter);
+
+	uba::SessionServerCreateInfo info(*UbaStorageServer, *UbaServer, LogWriter);
+	info.launchVisualizer = UbaJobProcessorOptions::bAutoLaunchVisualizer;
+	info.rootDir = *RootDir;
+	info.traceEnabled = true;
+	info.traceOutputFile = *UbaJobProcessorOptions::TraceFilename;
+	info.detailedTrace = UbaJobProcessorOptions::bDetailedTrace;
+	info.allowMemoryMaps = false; // Skip using memory maps
+	//info.remoteLogEnabled = true;
+	UbaSessionServer = CreateSessionServer(info);
+
+	CalculateKnownInputs();
+
+	UbaScheduler = Scheduler_Create(UbaSessionServer, MaxLocalParallelJobs, UbaJobProcessorOptions::bAllowProcessReuse);
+	Scheduler_Start(UbaScheduler);
+
+	HandleTaskQueueUpdated(TEXT("")); // Flush tasks into uba scheduler
+
+	Server_StartListen(UbaServer, uba::DefaultPort, nullptr); // Start listen so any helper on the LAN can join in
+
+	HordeAgentManager = MakeUnique<FUbaHordeAgentManager>(UbaServer);
+
+	UE_LOG(LogUbaController, Display, TEXT("Created UBA storage server: RootDir=%s"), *RootDir);
+}
+
+void FUbaJobProcessor::ShutDownUba()
+{
+	UE_LOG(LogUbaController, Display, TEXT("Shutting down UBA/Horde connection"));
+
+	HordeAgentManager = nullptr;
+
+	if (UbaSessionServer == nullptr)
+	{
+		return;
+	}
+
+	SessionServer_PrintSummary(UbaSessionServer);
+
+	Server_Stop(UbaServer);
+
+	Scheduler_Destroy(UbaScheduler);
+	DestroySessionServer(UbaSessionServer);
+	DestroyStorageServer(UbaStorageServer);
+	DestroyServer(UbaServer);
+
+	UbaScheduler = nullptr;
+	UbaSessionServer = nullptr;
+	UbaStorageServer = nullptr;
+	UbaServer = nullptr;
+}
+
+uint32 FUbaJobProcessor::Run()
+{
+	bIsWorkDone = false;
+	
+	uint32 LastTimeSinceHadJobs = FPlatformTime::Cycles();	
+
+	while (!bForceStop)
+	{
+		const float ElapsedSeconds = (FPlatformTime::Cycles() - LastTimeSinceHadJobs) * FPlatformTime::GetSecondsPerCycle();
+
+		uint32 queued = 0;
+		uint32 active = 0;
+		uint32 finished = 0;
+		if (UbaScheduler)
+		{
+			Scheduler_GetStats(UbaScheduler, queued, active, finished);
+		}
+
+		// We don't want to hog up Horde resources.
+		if (bShouldProcessJobs && ElapsedSeconds > UbaJobProcessorOptions::MaxTimeWithoutTasks && (queued + active) == 0)
+		{
+			// If we're optimizing job starting, we only want to shutdown UBA if all the processes have terminated
+			bShouldProcessJobs = false;
+			ShutDownUba();
+		}
+
+		// Check if we have new tasks to process
+		if ((ControllerModule.HasTasksDispatchedOrPending() || (queued + active) != 0))
+		{
+			if (!bShouldProcessJobs)
+			{
+				// We have new tasks. Start processing again
+				StartUba();
+
+				bShouldProcessJobs = true;
+			}
+
+			LastTimeSinceHadJobs = FPlatformTime::Cycles();
+		}
+
+		if (bShouldProcessJobs)
+		{
+			HordeAgentManager->SetTargetCoreCount(queued);
+			
+			// TODO: Not sure this is a good idea in a cooking scenario where number of queued processes are going up and down
+			SessionServer_SetMaxRemoteProcessCount(UbaSessionServer, queued);
+		}
+
+		FPlatformProcess::Sleep(UbaJobProcessorOptions::SleepTimeBetweenActions);
+	}
+
+	bIsWorkDone = true;
+	return 0;
+}
+
+void FUbaJobProcessor::StartThread()
+{
+	Thread = FRunnableThread::Create(this, TEXT("UbaJobProcessor"), 0, TPri_SlightlyBelowNormal, FPlatformAffinity::GetPoolThreadMask());
+}
+
+bool FUbaJobProcessor::ProcessOutputFile(FTask* CompileTask)
+{
+	//TODO: This method is mostly taken from the other Distribution controllers
+	// As we get an explicit callback when the process ends, we should be able to simplify this to just check if the file exists
+	IPlatformFile& PlatformFile = FPlatformFileManager::Get().GetPlatformFile();
+	IFileManager& FileManager = IFileManager::Get();
+
+	constexpr uint64 VersionAndFileSizeSize = sizeof(uint32) + sizeof(uint64);	
+	if (ensure(CompileTask && PlatformFile.FileExists(*CompileTask->CommandData.OutputFileName) &&
+		FileManager.FileSize(*CompileTask->CommandData.OutputFileName) > VersionAndFileSizeSize))
+	{
+		const TUniquePtr<FArchive> OutputFilePtr(FileManager.CreateFileReader(*CompileTask->CommandData.OutputFileName, FILEREAD_Silent));
+		if (ensure(OutputFilePtr))
+		{
+			FArchive& OutputFile = *OutputFilePtr;
+			int32 OutputVersion;
+			OutputFile << OutputVersion; // NOTE (SB): Do not care right now about the version.
+			int64 FileSize = 0;
+			OutputFile << FileSize;
+
+			// NOTE (SB): Check if we received the full file yet.
+			if (ensure(OutputFile.TotalSize() >= FileSize))
+			{
+				FTaskResponse TaskCompleted;
+				TaskCompleted.ID = CompileTask->ID;
+				TaskCompleted.ReturnCode = 0;
+						
+				ControllerModule.ReportJobProcessed(TaskCompleted, CompileTask);
+			}
+			else
+			{
+				UE_LOG(LogUbaController, Error, TEXT("Output file size is not correct [%s] | Expected Size [%lld] : => Actual Size : [%lld]"), *CompileTask->CommandData.OutputFileName, OutputFile.TotalSize(), FileSize);
+				return false;
+			}
+		}
+		else
+		{
+			UE_LOG(LogUbaController, Error, TEXT("Failed open for read Output File [%s]"), *CompileTask->CommandData.OutputFileName);
+			return false;
+		}
+	}
+	else
+	{
+		const FString OutputFileName = CompileTask != nullptr ? CompileTask->CommandData.OutputFileName : TEXT("Invalid CompileTask, cannot retrieve name");
+		UE_LOG(LogUbaController, Error, TEXT("Output File [%s] is invalid or do not exists"), *OutputFileName);
+		return false;
+	}
+
+	return true;
+}
+
+void FUbaJobProcessor::HandleUbaJobFinished(FTask* CompileTask)
+{
+	const bool bWasSuccessful = ProcessOutputFile(CompileTask);
+	if (!bWasSuccessful)
+	{
+		// If it failed running locally, lets try Run it locally but outside Uba
+		// Signaling a jobs as complete when it wasn't done, will cause a rerun on local worker as fallback
+		// because there is not output file for this job
+
+		FTaskResponse TaskCompleted;
+		TaskCompleted.ID = CompileTask->ID;
+		TaskCompleted.ReturnCode = 0;
+		ControllerModule.ReportJobProcessed(TaskCompleted, CompileTask);
+	}
+}
+
+void FUbaJobProcessor::HandleTaskQueueUpdated(const FString& InputFileName)
+{
+	if (!UbaScheduler)
+	{
+		return;
+	}
+
+	while (true)
+	{
+		FTask* Task = nullptr;
+		if (!ControllerModule.PendingRequestedCompilationTasks.Dequeue(Task) || !Task)
+			break;
+		RunTaskWithUba(Task);
+	}
+}
+
+bool FUbaJobProcessor::HasJobsInFlight() const
+{
+	if (!UbaScheduler)
+	{
+		return false;
+	}
+	uint32 queued = 0;
+	uint32 active = 0;
+	uint32 finished = 0;
+	Scheduler_GetStats(UbaScheduler, queued, active, finished);
+		
+	return queued != 0 || active != 0;
+}
\ No newline at end of file
diff --git a/Engine/Plugins/UbaController/Source/UbaController/Public/UbaControllerModule.h b/Engine/Plugins/UbaController/Source/UbaController/Public/UbaControllerModule.h
new file mode 100644
--- /dev/null
+++ b/Engine/Plugins/UbaController/Source/UbaController/Public/UbaControllerModule.h
@@ -0,0 +1,55 @@
+// Copyright Epic Games, Inc. All Rights Reserved.
+#pragma once
+
+#include "Containers/Queue.h"
+#include "DistributedBuildControllerInterface.h"
+
+class FUbaJobProcessor;
+struct FTaskCommandData;
+
+DECLARE_LOG_CATEGORY_EXTERN(LogUbaController, Log, Log);
+
+class FUbaControllerModule : public IDistributedBuildController , public TSharedFromThis<FUbaControllerModule>
+{
+public:
+	
+	FUbaControllerModule();
+	virtual ~FUbaControllerModule() override;
+	
+	UBACONTROLLER_API static FUbaControllerModule& Get();
+	UBACONTROLLER_API virtual bool IsSupported() override final;
+
+	UBACONTROLLER_API virtual const FString GetName() override final { return FString("UBA Controller"); };
+	
+	virtual void StartupModule() override final;
+	virtual void ShutdownModule() override final;
+	virtual void InitializeController() override final;
+	virtual FString CreateUniqueFilePath() override final;
+	virtual TFuture<FDistributedBuildTaskResult> EnqueueTask(const FTaskCommandData& CommandData) override final;
+	
+	void ReportJobProcessed(const FTaskResponse& InTaskResponse, FTask* CompileTask);
+	void CleanWorkingDirectory() const;
+
+	const FString& GetRootWorkingDirectory() const { return RootWorkingDirectory; }
+	const FString& GetWorkingDirectory() const { return WorkingDirectory; }
+
+	bool HasTasksDispatchedOrPending() const;
+
+	// Queue of tasks submitted by the engine, but not yet dispatched to the controller.
+	TQueue<FTask*, EQueueMode::Mpsc> PendingRequestedCompilationTasks;
+
+private:
+	void LoadDependencies();
+
+	bool bSupported;
+	bool bModuleInitialized;
+	bool bControllerInitialized;
+	
+	FString RootWorkingDirectory;
+	FString WorkingDirectory;
+
+	TAtomic<int32> NextFileID;
+	TAtomic<int32> NextTaskID;
+
+	TSharedPtr<FUbaJobProcessor> JobDispatcherThread;
+};
diff --git a/Engine/Plugins/UbaController/Source/UbaController/Public/UbaHordeAgent.h b/Engine/Plugins/UbaController/Source/UbaController/Public/UbaHordeAgent.h
new file mode 100644
--- /dev/null
+++ b/Engine/Plugins/UbaController/Source/UbaController/Public/UbaHordeAgent.h
@@ -0,0 +1,50 @@
+// Copyright Epic Games, Inc. All Rights Reserved.
+#pragma once
+
+#include <memory>
+
+#include "UbaHordeMetaClient.h"
+
+class FComputeSocket;
+class FAgentMessageChannel;
+
+DECLARE_LOG_CATEGORY_EXTERN(LogUbaHordeAgent, Log, Log);
+
+class FUbaHordeAgent
+{
+public:
+	FUbaHordeAgent(const FHordeRemoteMachineInfo &MachineInfo);
+
+	~FUbaHordeAgent();
+
+	/* Start spinning the ping thread, and receive/send threads. */
+	bool BeginCommunication();
+
+	/* Upload the proper executables/binaries. The Locator is just the contents of the bundle ref file.
+	 * e.g. Locator = abcdef1234567890abcdef1234567890abcdef12@34567890abcdef1234567890#123 */
+	bool UploadBinaries(const FString& BundleDirectory, const char* BundleLocator);
+
+	void Execute(const char* Exe, const char** Args, size_t NumArgs, const char* WorkingDir, const char** EnvVars, size_t NumEnvVars, bool bUseWine = false);
+
+	void CloseConnection(); 
+
+	bool IsValid();
+
+	// Reads output from the child channel and reports them to the log of the calling process.
+	void PollReports();
+
+	inline const FHordeRemoteMachineInfo& GetMachineInfo() const
+	{
+		return MachineInfo;
+	}
+
+private:
+	/* For Horde integration. Once all these are initialized, we're good to using UBA */
+	TUniquePtr<FComputeSocket> HordeComputeSocket;
+	TUniquePtr<FAgentMessageChannel> AgentChannel;
+	TUniquePtr<FAgentMessageChannel> ChildChannel;
+
+	bool bIsValid;
+	bool bHasErrors;
+	FHordeRemoteMachineInfo MachineInfo;
+};
diff --git a/Engine/Plugins/UbaController/Source/UbaController/Public/UbaHordeAgentManager.h b/Engine/Plugins/UbaController/Source/UbaController/Public/UbaHordeAgentManager.h
new file mode 100644
--- /dev/null
+++ b/Engine/Plugins/UbaController/Source/UbaController/Public/UbaHordeAgentManager.h
@@ -0,0 +1,48 @@
+// Copyright Epic Games, Inc. All Rights Reserved.
+#pragma once
+
+#include "Containers/Array.h"
+#include "HAL/Thread.h"
+#include "UbaHordeAgent.h"
+
+class FEvent;
+namespace uba { class NetworkServer; }
+
+class FUbaHordeAgentManager
+{
+public:
+	FUbaHordeAgentManager(uba::NetworkServer* InServer);
+	~FUbaHordeAgentManager();
+
+	void SetTargetCoreCount(uint32 Count);
+
+private:
+	void RequestAgent();
+	void ParseConfig();
+
+	struct FHordeAgentWrapper
+	{
+		FThread Thread;
+		FEvent* ShouldExit;
+	};
+
+	uba::NetworkServer* UbaServer;
+
+	FString Url;
+	FString Pool;
+	FString Oidc;
+	uint32 MaxCores = 0;
+
+	TUniquePtr<FUbaHordeMetaClient> HordeMetaClient;
+
+	FCriticalSection UbaAgentBundleFilePathLock;
+	FString UbaAgentBundleFilePath;
+
+	FCriticalSection AgentsLock;
+	TArray<TUniquePtr<FHordeAgentWrapper>> Agents;
+
+	TAtomic<uint64> LastRequestFailTime;
+	TAtomic<uint32> TargetCoreCount;
+	TAtomic<uint32> EstimatedCoreCount;
+	TAtomic<bool> AskForAgents;
+};
diff --git a/Engine/Plugins/UbaController/Source/UbaController/Public/UbaHordeComputeTransport.h b/Engine/Plugins/UbaController/Source/UbaController/Public/UbaHordeComputeTransport.h
new file mode 100644
--- /dev/null
+++ b/Engine/Plugins/UbaController/Source/UbaController/Public/UbaHordeComputeTransport.h
@@ -0,0 +1,34 @@
+// Copyright Epic Games, Inc. All Rights Reserved.
+#pragma once
+
+#include "Compute/ComputeTransport.h"
+#include "UbaHordeMetaClient.h"
+
+class FSocket;
+
+class FUbaHordeComputeTransport : public FComputeTransport
+{
+public:
+	// The constructor just performs a connect.
+	FUbaHordeComputeTransport(const FHordeRemoteMachineInfo &Info, bool& InHasErrors);
+	~FUbaHordeComputeTransport() override;
+	
+	// Sends data to the remote
+	size_t Send(const void* Data, size_t Size) override;
+
+	// Receives data from the remote
+	size_t Recv(void* Data, size_t Size) override;
+
+	// Indicates to the remote that no more data will be sent.
+	void MarkComplete() override;
+
+	// Indicates that no more data will be sent or received, and that any blocking reads/writes should stop.
+	void Close() override;
+
+	bool IsValid();
+
+private:
+	FSocket* Socket;
+	bool bIsClosed;
+	bool& bHasErrors;
+};
diff --git a/Engine/Plugins/UbaController/Source/UbaController/Public/UbaHordeMetaClient.h b/Engine/Plugins/UbaController/Source/UbaController/Public/UbaHordeMetaClient.h
new file mode 100644
--- /dev/null
+++ b/Engine/Plugins/UbaController/Source/UbaController/Public/UbaHordeMetaClient.h
@@ -0,0 +1,64 @@
+// Copyright Epic Games, Inc. All Rights Reserved.
+#pragma once
+
+#include "HttpModule.h"
+#include "Async/Future.h"
+#include "Containers/UnrealString.h"
+#include "Interfaces/IHttpRequest.h"
+#include "Interfaces/IHttpResponse.h"
+#include "HordeHttpClient.h"
+
+DECLARE_LOG_CATEGORY_EXTERN(LogUbaHorde, Log, Log);
+
+#define HORDE_NONCE_SIZE 64
+
+// When asking the Horde server for a machine, this will be returned when the HTTP
+// reponse comes back (after a machine request). If the response didn't come through
+// (i.e., Horde was unable to assign a machine or something), then:
+// - Ip will be set to ""
+// - Port will be set to 0xFFFF
+// - Nonce will be all 0
+struct FHordeRemoteMachineInfo
+{
+	FString Ip;
+	uint16 Port;
+	uint16 LogicalCores;
+	uint8 Nonce[HORDE_NONCE_SIZE];
+	bool bRunsWindowOS;
+};
+
+// This encapsulates the mechanism of talking to the Horde "meta server". The server which
+// can grant us access to remote machines.
+class FUbaHordeMetaClient
+{
+public:
+	// We return the HttpResponse in case more information could be used out of it later.
+	using HordeMachinePromise = TPromise<TTuple<FHttpResponsePtr, FHordeRemoteMachineInfo>>;
+	
+	// Construct the client with a single URL that includes the protocol, IP address, and port in the common format "PROTOCOL://IP:PORT".
+	FUbaHordeMetaClient(const FStringView& HordeServerUrl, const FStringView& InOAuthProviderIdentifier);
+
+	// Local Horde server uses port 5000 by default.
+	// Production Horde servers at horde.devtools.epicgames.com use the standard HTTPS port 443.
+	FUbaHordeMetaClient(const FStringView& HordeServerIp, uint16 HordeServerPort, bool bInConnectWithAuthentication, const FStringView& OAuthProviderIdentifier);
+
+	bool RefreshHttpClient();
+
+	// This will make a request to Horde for a remote machine to do work on.
+	// Example of actually getting the FHordeRemoteMachineInfo struct:
+	// ```
+	// auto Promise = HordeServer.RequestMachine();
+	// auto Future = Promise->GetFuture();
+	// ... // Can to asynchronous work
+	// auto Future.Wait(); // Wait for the response to arrive
+	// FHordeRemoteMachineInfo MachineInfo = Future.Get().Value;
+	TSharedPtr<HordeMachinePromise, ESPMode::ThreadSafe> RequestMachine(const FString& PoolId, const FString& Machine = "default");
+
+private:
+	void ParseConfig(const FString& IniFilename, const FString& IniSection);
+
+	const FString ServerUrl;
+	const bool bConnectWithAuthentication;
+	TUniquePtr<FHordeHttpClient> HttpClient;
+	FString OAuthProviderIdentifier;
+};
diff --git a/Engine/Plugins/UbaController/Source/UbaController/Public/UbaJobProcessor.h b/Engine/Plugins/UbaController/Source/UbaController/Public/UbaJobProcessor.h
new file mode 100644
--- /dev/null
+++ b/Engine/Plugins/UbaController/Source/UbaController/Public/UbaJobProcessor.h
@@ -0,0 +1,77 @@
+// Copyright Epic Games, Inc. All Rights Reserved.
+#pragma once
+
+#include "HAL/PlatformAtomics.h"
+#include "HAL/Runnable.h"
+#include "Templates/SharedPointer.h"
+#include "Templates/UniquePtr.h"
+#include "UbaExports.h"
+
+class FUbaControllerModule;
+class FUbaHordeAgentManager;
+struct FTask;
+
+class FUbaJobProcessor : public FRunnable, public TSharedFromThis<FUbaJobProcessor>
+{
+public:
+	FUbaJobProcessor(FUbaControllerModule& InControllerModule);
+
+	virtual ~FUbaJobProcessor() override;
+	
+	/** Main loop */
+	virtual uint32 Run() override;
+	
+	/** Aborts the main loop as soon as possible */
+	virtual void Stop() override { bForceStop = true; };
+	
+	/** Creates the threads and starts the main loop */
+	void StartThread();
+	bool ProcessOutputFile(FTask* CompileTask);
+
+	void HandleUbaJobFinished(FTask* CompileTask);
+
+	/** Used to know when this thread has finished the main loop */
+	bool IsWorkDone() const { return bIsWorkDone;};
+
+	void HandleTaskQueueUpdated(const FString& InputFileName);
+
+	bool HasJobsInFlight() const;
+
+private:
+
+	void CalculateKnownInputs();
+	void RunTaskWithUba(FTask* Task);
+
+	void StartUba();
+	void ShutDownUba();
+
+	/** The runnable thread */
+	FRunnableThread* Thread;
+	
+	FUbaControllerModule& ControllerModule;
+
+	int32 MaxLocalParallelJobs = 0;
+	
+	/** Used to abort the current processing loop */
+	TAtomic<bool> bForceStop;
+	
+	/** Used to abort the current processing loop */
+	uint32 LastTimeCheckedForTasks;
+	
+	/** Used to stop the processing loop in the next main loop */
+	TAtomic<bool> bShouldProcessJobs;
+
+	/** Set to true when the main loop finishes*/
+	bool bIsWorkDone;
+
+	uba::NetworkServer* UbaServer = nullptr;
+	uba::Storage* UbaStorageServer = nullptr;
+	uba::SessionServer* UbaSessionServer = nullptr;
+	uba::Scheduler* UbaScheduler = nullptr;
+	TUniquePtr<FUbaHordeAgentManager> HordeAgentManager;
+
+	uint32 KnownInputsCount = 0;
+	TArray<uba::tchar> KnownInputsBuffer;
+
+	uba::CallbackLogWriter LogWriter;
+};
\ No newline at end of file
diff --git a/Engine/Plugins/UbaController/Source/UbaController/UbaController.build.cs b/Engine/Plugins/UbaController/Source/UbaController/UbaController.build.cs
new file mode 100644
--- /dev/null
+++ b/Engine/Plugins/UbaController/Source/UbaController/UbaController.build.cs
@@ -0,0 +1,39 @@
+// Copyright Epic Games, Inc. All Rights Reserved.
+
+using System.IO;
+using UnrealBuildTool;
+
+public class UbaController : ModuleRules
+{
+	public UbaController(ReadOnlyTargetRules TargetRules) : base(TargetRules)
+	{
+		PrivateDependencyModuleNames.AddRange(new string[]
+		{
+			"Core",
+			"DistributedBuildInterface",
+			"Projects",
+			"RenderCore",
+			"TargetPlatform",
+			"Horde",
+			"HTTP",
+			"Sockets",
+			"Json"
+		});
+
+		string Arch = TargetRules.Architecture.WindowsName == "arm64ec" ? "arm64" : TargetRules.Architecture.WindowsName;
+		string UbaSourcePath = Path.Combine(EngineDirectory, "Source", "Programs", "UnrealBuildAccelerator");
+
+		PublicIncludePaths.Add(UbaSourcePath);
+		PublicIncludePaths.Add(Path.Combine(UbaSourcePath, "Core", "Public"));
+		PublicIncludePaths.Add(Path.Combine(UbaSourcePath, "Common", "Public"));
+        PublicIncludePaths.Add(Path.Combine(UbaSourcePath, "Host", "Public"));
+        PublicIncludePathModuleNames.Add("BLAKE3");
+
+		string UbaBinariesPath = Path.Combine(EngineDirectory, "Binaries", "Win64", "UnrealBuildAccelerator", Arch);
+
+		PublicAdditionalLibraries.Add(Path.Combine(UbaBinariesPath, "UbaHost.lib"));
+		PrivateRuntimeLibraryPaths.Add(Path.Combine(ModuleDirectory, "lib"));
+
+		PublicDelayLoadDLLs.Add("UbaHost.dll");
+	}
+}
\ No newline at end of file
diff --git a/Engine/Plugins/UbaController/UbaController.uplugin b/Engine/Plugins/UbaController/UbaController.uplugin
new file mode 100644
--- /dev/null
+++ b/Engine/Plugins/UbaController/UbaController.uplugin
@@ -0,0 +1,30 @@
+{
+	"FileVersion" : 1,
+	"Version" : 1,
+	"VersionName" : "0.1",
+	"FriendlyName" : "UBA Controller",
+	"Description" : "Adds support for shader compiling distribution using UnrealBuildAccelerator (UBA)",
+	"Category" : "Build Distribution",
+	"CreatedBy" : "Epic Games, Inc.",
+	"CreatedByURL" : "https://epicgames.com",
+	"DocsURL" : "",
+	"MarketplaceURL" : "",
+	"SupportURL" : "",
+	"IsExperimentalVersion": true,
+	"EnabledByDefault" : false,
+	"CanContainContent" : false,
+	"IsBetaVersion" : false,
+	"Installed" : false,
+
+	"Modules" :
+	[
+		{
+			"Name" : "UbaController",
+			"Type" : "Editor",
+			"LoadingPhase" : "EarliestPossible",
+			"PlatformAllowList": ["Win64"]
+		}
+	],
+
+	"SupportedTargetPlatforms": ["Win64"]
+}
\ No newline at end of file
