[UBAController]
* Removed comment causing error on farm

diff --git a/Engine/Plugins/UbaController/Source/UbaController/Public/UbaHordeMetaClient.h b/Engine/Plugins/UbaController/Source/UbaController/Public/UbaHordeMetaClient.h
--- a/Engine/Plugins/UbaController/Source/UbaController/Public/UbaHordeMetaClient.h
+++ b/Engine/Plugins/UbaController/Source/UbaController/Public/UbaHordeMetaClient.h
@@ -39,7 +39,6 @@ public:
 	FUbaHordeMetaClient(const FStringView& HordeServerUrl, const FStringView& InOAuthProviderIdentifier);
 
 	// Local Horde server uses port 5000 by default.
-	// Production Horde servers at horde.devtools.epicgames.com use the standard HTTPS port 443.
 	FUbaHordeMetaClient(const FStringView& HordeServerIp, uint16 HordeServerPort, bool bInConnectWithAuthentication, const FStringView& OAuthProviderIdentifier);
 
 	bool RefreshHttpClient();
