[UBAController]
* Added cleanup of uba when shutting down
* Moved some code to a function instead of lambda to make it easier to put breakpoints etc.

diff --git a/Engine/Plugins/UbaController/Source/UbaController/Private/UbaHordeAgentManager.cpp b/Engine/Plugins/UbaController/Source/UbaController/Private/UbaHordeAgentManager.cpp
--- a/Engine/Plugins/UbaController/Source/UbaController/Private/UbaHordeAgentManager.cpp
+++ b/Engine/Plugins/UbaController/Source/UbaController/Private/UbaHordeAgentManager.cpp
@@ -161,192 +161,192 @@ void FUbaHordeAgentManager::RequestAgent()
 	FHordeAgentWrapper& Wrapper = *Agents.Emplace_GetRef(MakeUnique<FHordeAgentWrapper>());
 
 	Wrapper.ShouldExit = FGenericPlatformProcess::GetSynchEventFromPool(true);
-	Wrapper.Thread = FThread(
-		TEXT("HordeAgent"),
-		[this, WrapperPtr = &Wrapper]()
+	Wrapper.Thread = FThread(TEXT("HordeAgent"), [this, WrapperPtr = &Wrapper]() { ThreadAgent(*WrapperPtr); });
+}
+
+void FUbaHordeAgentManager::ThreadAgent(FHordeAgentWrapper& Wrapper)
+{
+	FEvent& ShouldExit = *Wrapper.ShouldExit;
+	TUniquePtr<FUbaHordeAgent> Agent;
+	bool bSuccess = false;
+
+	ON_SCOPE_EXIT
+	{
+		if (Agent)
 		{
-			FEvent& ShouldExit = *WrapperPtr->ShouldExit;
-			TUniquePtr<FUbaHordeAgent> Agent;
-			bool bSuccess = false;
+			Agent->CloseConnection();
+		}
 
-			ON_SCOPE_EXIT
-			{
-				if (Agent)
-				{
-					Agent->CloseConnection();
-				}
+		ShouldExit.Trigger();
+	};
 
-				ShouldExit.Trigger();
-			};
+	int MachineCoreCount = 0;
 
-			int MachineCoreCount = 0;
+	{
+		ON_SCOPE_EXIT{ EstimatedCoreCount -= 32; };
 
-			{
-				ON_SCOPE_EXIT{ EstimatedCoreCount -= 32; };
+		FScopeLock ScopeLock(&UbaAgentBundleFilePathLock);
+		if (UbaAgentBundleFilePath.IsEmpty())
+		{
+			FString RootWorkingDirectory = FPaths::Combine(FPlatformProcess::UserTempDir(), TEXT("UbaWorkingDir"));
 
-				FScopeLock ScopeLock(&UbaAgentBundleFilePathLock);
-				if (UbaAgentBundleFilePath.IsEmpty())
-				{
-					FString RootWorkingDirectory = FPaths::Combine(FPlatformProcess::UserTempDir(), TEXT("UbaWorkingDir"));
-
-					const FString UbaAgentFilePath = FPaths::Combine(GetUbaBinariesPath(), TEXT("UbaAgent.exe"));
-					UbaAgentBundleFilePath = FPaths::Combine(RootWorkingDirectory, TEXT("UbaAgent.Bundle.ref"));
-
-					if (!CreateHordeBundleFromFile(*UbaAgentFilePath, *UbaAgentBundleFilePath))
-					{
-						UE_LOG(LogUbaController, Error, TEXT("Failed to create Horde bundle for UbaAgent executable: %s"), *UbaAgentFilePath);
-						AskForAgents = false;
-						return;
-					}
-					UE_LOG(LogUbaController, Display, TEXT("Created Horde bundle for UbaAgent executable: %s"), *UbaAgentFilePath);
-				}
+			const FString UbaAgentFilePath = FPaths::Combine(GetUbaBinariesPath(), TEXT("UbaAgent.exe"));
+			UbaAgentBundleFilePath = FPaths::Combine(RootWorkingDirectory, TEXT("UbaAgent.Bundle.ref"));
 
-				if (!HordeMetaClient)
-				{
-					// Create Horde meta client right before we need it to make sure the CVar for the server URL has been read by now
-					HordeMetaClient = MakeUnique<FUbaHordeMetaClient>(Url, Oidc);
-					if (!HordeMetaClient->RefreshHttpClient())
-					{
-						UE_LOG(LogUbaController, Error, TEXT("Failed to create HttpClient for UbaAgent"));
-						AskForAgents = false;
-						return;
-					}
-				}
+			if (!CreateHordeBundleFromFile(*UbaAgentFilePath, *UbaAgentBundleFilePath))
+			{
+				UE_LOG(LogUbaController, Error, TEXT("Failed to create Horde bundle for UbaAgent executable: %s"), *UbaAgentFilePath);
+				AskForAgents = false;
+				return;
+			}
+			UE_LOG(LogUbaController, Display, TEXT("Created Horde bundle for UbaAgent executable: %s"), *UbaAgentFilePath);
+		}
 
-				if (!AskForAgents)
-				{
-					return;
-				}
+		if (!HordeMetaClient)
+		{
+			// Create Horde meta client right before we need it to make sure the CVar for the server URL has been read by now
+			HordeMetaClient = MakeUnique<FUbaHordeMetaClient>(Url, Oidc);
+			if (!HordeMetaClient->RefreshHttpClient())
+			{
+				UE_LOG(LogUbaController, Error, TEXT("Failed to create HttpClient for UbaAgent"));
+				AskForAgents = false;
+				return;
+			}
+		}
 
-				if (LastRequestFailTime == 0)
-				{
-					ScopeLock.Unlock();
-				}
-				else
-				{
-					// Try to reduce pressure on horde by not asking for machines more frequent than every 5 seconds if failed to retrieve last time
-					uint64 CurrentTime = FPlatformTime::Cycles64();
-					uint32 MsSinceLastFail = uint32((CurrentTime - LastRequestFailTime) * FPlatformTime::GetSecondsPerCycle() * 1000);
-					if (MsSinceLastFail < 5000)
-					{
-						if (ShouldExit.Wait(5000 - MsSinceLastFail))
-						{
-							return;
-						}
-					}
-				}
+		if (!AskForAgents)
+		{
+			return;
+		}
 
-				TSharedPtr<FUbaHordeMetaClient::HordeMachinePromise, ESPMode::ThreadSafe> Promise = HordeMetaClient->RequestMachine(Pool);
-				if (!Promise)
+		if (LastRequestFailTime == 0)
+		{
+			ScopeLock.Unlock();
+		}
+		else
+		{
+			// Try to reduce pressure on horde by not asking for machines more frequent than every 5 seconds if failed to retrieve last time
+			uint64 CurrentTime = FPlatformTime::Cycles64();
+			uint32 MsSinceLastFail = uint32((CurrentTime - LastRequestFailTime) * FPlatformTime::GetSecondsPerCycle() * 1000);
+			if (MsSinceLastFail < 5000)
+			{
+				if (ShouldExit.Wait(5000 - MsSinceLastFail))
 				{
-					//UE_LOG(LogUbaController, Error, TEXT("Failed to create Horde bundle for UbaAgent executable: %s"), *UbaAgentFilePath);
 					return;
 				}
-				TFuture<TTuple<FHttpResponsePtr, FHordeRemoteMachineInfo>> Future = Promise->GetFuture();
-				Future.Wait();
-				FHordeRemoteMachineInfo MachineInfo = Future.Get().Value;
+			}
+		}
 
-				// If the machine couldn't be assigned, just ignore this agent slot
-				if (MachineInfo.Ip == TEXT(""))
-				{
-					LastRequestFailTime = FPlatformTime::Cycles64();
-					return;
-				}
+		TSharedPtr<FUbaHordeMetaClient::HordeMachinePromise, ESPMode::ThreadSafe> Promise = HordeMetaClient->RequestMachine(Pool);
+		if (!Promise)
+		{
+			//UE_LOG(LogUbaController, Error, TEXT("Failed to create Horde bundle for UbaAgent executable: %s"), *UbaAgentFilePath);
+			return;
+		}
+		TFuture<TTuple<FHttpResponsePtr, FHordeRemoteMachineInfo>> Future = Promise->GetFuture();
+		Future.Wait();
+		FHordeRemoteMachineInfo MachineInfo = Future.Get().Value;
 
-				LastRequestFailTime = 0;
+		// If the machine couldn't be assigned, just ignore this agent slot
+		if (MachineInfo.Ip == TEXT(""))
+		{
+			LastRequestFailTime = FPlatformTime::Cycles64();
+			return;
+		}
 
-				ScopeLock.Unlock();
+		LastRequestFailTime = 0;
 
-				if (ShouldExit.Wait(0))
-				{
-					return;
-				}
+		ScopeLock.Unlock();
 
-				Agent = MakeUnique<FUbaHordeAgent>(MachineInfo);
+		if (ShouldExit.Wait(0))
+		{
+			return;
+		}
 
-				if (!Agent->BeginCommunication())
-				{
-					return;
-				}
+		Agent = MakeUnique<FUbaHordeAgent>(MachineInfo);
 
-				TArray<uint8> Locator;
-				if (!FFileHelper::LoadFileToArray(Locator, *UbaAgentBundleFilePath))
-				{
-					UE_LOG(LogUbaController, Error, TEXT("Cannot launch Horde processes for UBA controller because bundle path could not be found: %s"), *UbaAgentBundleFilePath);
-					return;
-				}
+		if (!Agent->BeginCommunication())
+		{
+			return;
+		}
+
+		TArray<uint8> Locator;
+		if (!FFileHelper::LoadFileToArray(Locator, *UbaAgentBundleFilePath))
+		{
+			UE_LOG(LogUbaController, Error, TEXT("Cannot launch Horde processes for UBA controller because bundle path could not be found: %s"), *UbaAgentBundleFilePath);
+			return;
+		}
 
-				Locator.Add('\0');
+		Locator.Add('\0');
 
-				FString BundleDirectory = FPaths::GetPath(UbaAgentBundleFilePath);
+		FString BundleDirectory = FPaths::GetPath(UbaAgentBundleFilePath);
 
-				if (ShouldExit.Wait(0))
-				{
-					return;
-				}
+		if (ShouldExit.Wait(0))
+		{
+			return;
+		}
 
-				if (!Agent->UploadBinaries(BundleDirectory, reinterpret_cast<const char*>(Locator.GetData())))
-				{
-					return;
-				}
+		if (!Agent->UploadBinaries(BundleDirectory, reinterpret_cast<const char*>(Locator.GetData())))
+		{
+			return;
+		}
 
-				uint32 ListenPort = 7001;
+		uint32 ListenPort = 7001;
 
-				// Start the UBA Agent that will connect to us, requesting for work
-				const std::string ListenPortArg = "-listen=" + std::to_string(ListenPort);
+		// Start the UBA Agent that will connect to us, requesting for work
+		const std::string ListenPortArg = "-listen=" + std::to_string(ListenPort);
 
-				const char* UbaAgentArgs[] =
-				{
-					ListenPortArg.c_str(),
-					"-nopoll",				// -nopoll recommended when running on remote Horde agents to make sure they exit after completion. Otherwise, it keeps running.
-					"-listenTimeout=20",	// increase timeout for the agent to listen from 5 to 20 seconds to give the editor enough time to establish the connection
-				};
+		const char* UbaAgentArgs[] =
+		{
+			ListenPortArg.c_str(),
+			"-nopoll",				// -nopoll recommended when running on remote Horde agents to make sure they exit after completion. Otherwise, it keeps running.
+			"-listenTimeout=20",	// increase timeout for the agent to listen from 5 to 20 seconds to give the editor enough time to establish the connection
+		};
 
-				// If the machine does not run Windows, enable the compatibility layer Wine to run UbaAgent.exe on POSIX systems
-				const bool bRunsWindowsOS = Agent->GetMachineInfo().bRunsWindowOS;
-				const bool bUseWine = !bRunsWindowsOS;
+		// If the machine does not run Windows, enable the compatibility layer Wine to run UbaAgent.exe on POSIX systems
+		const bool bRunsWindowsOS = Agent->GetMachineInfo().bRunsWindowOS;
+		const bool bUseWine = !bRunsWindowsOS;
 
-				if (ShouldExit.Wait(0))
-				{
-					return;
-				}
+		if (ShouldExit.Wait(0))
+		{
+			return;
+		}
 
-				Agent->Execute("UbaAgent.exe", UbaAgentArgs, UE_ARRAY_COUNT(UbaAgentArgs), nullptr, nullptr, 0, bUseWine);
+		Agent->Execute("UbaAgent.exe", UbaAgentArgs, UE_ARRAY_COUNT(UbaAgentArgs), nullptr, nullptr, 0, bUseWine);
 
-				// Add this machine as client to the remote agent
-				const FString& IpAddress = Agent->GetMachineInfo().Ip;
-				auto IpAddressStr = StringCast<uba::tchar>(*IpAddress);
-				const bool bAddClientSuccess = Server_AddClient(UbaServer, IpAddressStr.Get(), ListenPort, nullptr);
+		// Add this machine as client to the remote agent
+		const FString& IpAddress = Agent->GetMachineInfo().Ip;
+		auto IpAddressStr = StringCast<uba::tchar>(*IpAddress);
+		const bool bAddClientSuccess = Server_AddClient(UbaServer, IpAddressStr.Get(), ListenPort, nullptr);
 
-				if (!bAddClientSuccess)
-				{
-					UE_LOG(LogUbaController, Error, TEXT("Server_AddClient(%s:%d) failed"), *IpAddress, ListenPort);
-					return;
-				}
+		if (!bAddClientSuccess)
+		{
+			UE_LOG(LogUbaController, Error, TEXT("Server_AddClient(%s:%d) failed"), *IpAddress, ListenPort);
+			return;
+		}
 
-				// Log remote execution
-				FString UbaAgentCmdArgs = TEXT("UbaAgent.exe");
-				for (const char* Arg : UbaAgentArgs)
-				{
-					UbaAgentCmdArgs += TEXT(" ");
-					UbaAgentCmdArgs += ANSI_TO_TCHAR(Arg);
-				}
-				UE_LOG(LogUbaController, Log, TEXT("Remote execution on Horde machine [%s:%d]: %s"), *IpAddress, ListenPort, *UbaAgentCmdArgs);
+		// Log remote execution
+		FString UbaAgentCmdArgs = TEXT("UbaAgent.exe");
+		for (const char* Arg : UbaAgentArgs)
+		{
+			UbaAgentCmdArgs += TEXT(" ");
+			UbaAgentCmdArgs += ANSI_TO_TCHAR(Arg);
+		}
+		UE_LOG(LogUbaController, Log, TEXT("Remote execution on Horde machine [%s:%d]: %s"), *IpAddress, ListenPort, *UbaAgentCmdArgs);
 
-				MachineCoreCount = MachineInfo.LogicalCores;
-				EstimatedCoreCount += MachineCoreCount;
-			}
+		MachineCoreCount = MachineInfo.LogicalCores;
+		EstimatedCoreCount += MachineCoreCount;
+	}
 
-			while (!ShouldExit.Wait(100))
-			{
-				if (!Agent->IsValid())
-					break;
-				if (UbaControllerModule::bHordeForwardAgentLogs)
-					Agent->PollReports();
-			}
+	while (!ShouldExit.Wait(100))
+	{
+		if (!Agent->IsValid())
+			break;
+		if (UbaControllerModule::bHordeForwardAgentLogs)
+			Agent->PollReports();
+	}
 
-			EstimatedCoreCount -= MachineCoreCount;
-		});
+	EstimatedCoreCount -= MachineCoreCount;
 }
 
 void FUbaHordeAgentManager::ParseConfig()
diff --git a/Engine/Plugins/UbaController/Source/UbaController/Private/UbaJobProcessor.cpp b/Engine/Plugins/UbaController/Source/UbaController/Private/UbaJobProcessor.cpp
--- a/Engine/Plugins/UbaController/Source/UbaController/Private/UbaJobProcessor.cpp
+++ b/Engine/Plugins/UbaController/Source/UbaController/Private/UbaJobProcessor.cpp
@@ -352,6 +352,8 @@ uint32 FUbaJobProcessor::Run()
 		FPlatformProcess::Sleep(UbaJobProcessorOptions::SleepTimeBetweenActions);
 	}
 
+	ShutDownUba();
+
 	bIsWorkDone = true;
 	return 0;
 }
diff --git a/Engine/Plugins/UbaController/Source/UbaController/Public/UbaHordeAgentManager.h b/Engine/Plugins/UbaController/Source/UbaController/Public/UbaHordeAgentManager.h
--- a/Engine/Plugins/UbaController/Source/UbaController/Public/UbaHordeAgentManager.h
+++ b/Engine/Plugins/UbaController/Source/UbaController/Public/UbaHordeAgentManager.h
@@ -17,15 +17,16 @@ public:
 	void SetTargetCoreCount(uint32 Count);
 
 private:
-	void RequestAgent();
-	void ParseConfig();
-
 	struct FHordeAgentWrapper
 	{
 		FThread Thread;
 		FEvent* ShouldExit;
 	};
 
+	void RequestAgent();
+	void ThreadAgent(FHordeAgentWrapper& Wrapper);
+	void ParseConfig();
+
 	uba::NetworkServer* UbaServer;
 
 	FString Url;
