[UBAController]
* Fixed math for calculating how many remote cores are desired
* Fixed bug where agent disconnects could be missed

diff --git a/Engine/Plugins/UbaController/Source/UbaController/Private/UbaHordeAgent.cpp b/Engine/Plugins/UbaController/Source/UbaController/Private/UbaHordeAgent.cpp
--- a/Engine/Plugins/UbaController/Source/UbaController/Private/UbaHordeAgent.cpp
+++ b/Engine/Plugins/UbaController/Source/UbaController/Private/UbaHordeAgent.cpp
@@ -181,7 +181,7 @@ bool FUbaHordeAgent::IsValid()
 	return bIsValid && !bHasErrors;
 }
 
-void FUbaHordeAgent::PollReports()
+void FUbaHordeAgent::Poll(bool LogReports)
 {
 	EAgentMessageType Type = EAgentMessageType::None;
 	constexpr int32 ReadresponseTimeoutMS = 100;
@@ -191,7 +191,7 @@ void FUbaHordeAgent::PollReports()
 		{
 		case EAgentMessageType::ExecuteOutput:
 		{
-			if (ChildChannel->GetResponseSize() > 0)
+			if (LogReports && ChildChannel->GetResponseSize() > 0)
 			{
 				// Convert raw buffer to dynamic array for modification
 				const ANSICHAR* ResponseDataRaw = reinterpret_cast<const ANSICHAR*>(ChildChannel->GetResponseData());
@@ -219,9 +219,12 @@ void FUbaHordeAgent::PollReports()
 		{
 			if (ChildChannel->GetResponseSize() == sizeof(int32))
 			{
-				const int32* ResponseData = reinterpret_cast<const int32*>(ChildChannel->GetResponseData());
-				const int32 ExecuteExitCode = ResponseData[0];
-				UE_LOG(LogUbaHorde, Log, TEXT("Response [ExecuteResult]: ExitCode=%d"), ExecuteExitCode);
+				if (LogReports)
+				{
+					const int32* ResponseData = reinterpret_cast<const int32*>(ChildChannel->GetResponseData());
+					const int32 ExecuteExitCode = ResponseData[0];
+					UE_LOG(LogUbaHorde, Log, TEXT("Response [ExecuteResult]: ExitCode=%d"), ExecuteExitCode);
+				}
 				bIsValid = false;
 			}
 		}
diff --git a/Engine/Plugins/UbaController/Source/UbaController/Private/UbaHordeAgentManager.cpp b/Engine/Plugins/UbaController/Source/UbaController/Private/UbaHordeAgentManager.cpp
--- a/Engine/Plugins/UbaController/Source/UbaController/Private/UbaHordeAgentManager.cpp
+++ b/Engine/Plugins/UbaController/Source/UbaController/Private/UbaHordeAgentManager.cpp
@@ -338,12 +338,9 @@ void FUbaHordeAgentManager::ThreadAgent(FHordeAgentWrapper& Wrapper)
 		EstimatedCoreCount += MachineCoreCount;
 	}
 
-	while (!ShouldExit.Wait(100))
+	while (Agent->IsValid() && !ShouldExit.Wait(100))
 	{
-		if (!Agent->IsValid())
-			break;
-		if (UbaControllerModule::bHordeForwardAgentLogs)
-			Agent->PollReports();
+		Agent->Poll(UbaControllerModule::bHordeForwardAgentLogs);
 	}
 
 	EstimatedCoreCount -= MachineCoreCount;
diff --git a/Engine/Plugins/UbaController/Source/UbaController/Private/UbaJobProcessor.cpp b/Engine/Plugins/UbaController/Source/UbaController/Private/UbaJobProcessor.cpp
--- a/Engine/Plugins/UbaController/Source/UbaController/Private/UbaJobProcessor.cpp
+++ b/Engine/Plugins/UbaController/Source/UbaController/Private/UbaJobProcessor.cpp
@@ -343,10 +343,12 @@ uint32 FUbaJobProcessor::Run()
 
 		if (bShouldProcessJobs)
 		{
-			HordeAgentManager->SetTargetCoreCount(queued);
+			int32 TargetCoreCount = FMath::Max(0, int32(queued + active) - MaxLocalParallelJobs);
+
+			HordeAgentManager->SetTargetCoreCount(TargetCoreCount);
 			
 			// TODO: Not sure this is a good idea in a cooking scenario where number of queued processes are going up and down
-			SessionServer_SetMaxRemoteProcessCount(UbaSessionServer, queued);
+			SessionServer_SetMaxRemoteProcessCount(UbaSessionServer, TargetCoreCount);
 		}
 
 		FPlatformProcess::Sleep(UbaJobProcessorOptions::SleepTimeBetweenActions);
diff --git a/Engine/Plugins/UbaController/Source/UbaController/Public/UbaHordeAgent.h b/Engine/Plugins/UbaController/Source/UbaController/Public/UbaHordeAgent.h
--- a/Engine/Plugins/UbaController/Source/UbaController/Public/UbaHordeAgent.h
+++ b/Engine/Plugins/UbaController/Source/UbaController/Public/UbaHordeAgent.h
@@ -31,7 +31,7 @@ public:
 	bool IsValid();
 
 	// Reads output from the child channel and reports them to the log of the calling process.
-	void PollReports();
+	void Poll(bool LogReports);
 
 	inline const FHordeRemoteMachineInfo& GetMachineInfo() const
 	{
