[UBAController]
* Downgraded error to display when failing to connect to uba client since it should not cause build to log error

diff --git a/Engine/Plugins/UbaController/Source/UbaController/Private/UbaHordeAgentManager.cpp b/Engine/Plugins/UbaController/Source/UbaController/Private/UbaHordeAgentManager.cpp
--- a/Engine/Plugins/UbaController/Source/UbaController/Private/UbaHordeAgentManager.cpp
+++ b/Engine/Plugins/UbaController/Source/UbaController/Private/UbaHordeAgentManager.cpp
@@ -321,7 +321,7 @@ void FUbaHordeAgentManager::ThreadAgent(FHordeAgentWrapper& Wrapper)
 
 		if (!bAddClientSuccess)
 		{
-			UE_LOG(LogUbaController, Error, TEXT("Server_AddClient(%s:%d) failed"), *IpAddress, ListenPort);
+			UE_LOG(LogUbaController, Display, TEXT("Server_AddClient(%s:%d) failed"), *IpAddress, ListenPort);
 			return;
 		}
 
