[UBAController]
* Fixed so multiprocess cook works by making sure each instance has its own casdb..
* Fixed so process logs are disabled by default
* Fixed so uba errors are outputted as errors in the UE log
* Added system wide lock around oidc request in Horde meta client to resolve problem with ue worker processes failing to fetch token

diff --git a/Engine/Plugins/UbaController/Source/UbaController/Private/UbaControllerModule.cpp b/Engine/Plugins/UbaController/Source/UbaController/Private/UbaControllerModule.cpp
--- a/Engine/Plugins/UbaController/Source/UbaController/Private/UbaControllerModule.cpp
+++ b/Engine/Plugins/UbaController/Source/UbaController/Private/UbaControllerModule.cpp
@@ -67,6 +67,11 @@ bool FUbaControllerModule::IsSupported()
 
 void FUbaControllerModule::CleanWorkingDirectory() const
 {
+	if (UE::GetMultiprocessId() != 0) // Only director is allowed to clean
+	{
+		return;
+	}
+
 	IFileManager& FileManager = IFileManager::Get();
 	
 	if (!RootWorkingDirectory.IsEmpty())
diff --git a/Engine/Plugins/UbaController/Source/UbaController/Private/UbaHordeAgentManager.cpp b/Engine/Plugins/UbaController/Source/UbaController/Private/UbaHordeAgentManager.cpp
--- a/Engine/Plugins/UbaController/Source/UbaController/Private/UbaHordeAgentManager.cpp
+++ b/Engine/Plugins/UbaController/Source/UbaController/Private/UbaHordeAgentManager.cpp
@@ -31,8 +31,9 @@ namespace UbaControllerModule
 		ECVF_ReadOnly); // Must be set on start-up, e.g. via config ini
 }
 
-FUbaHordeAgentManager::FUbaHordeAgentManager(uba::NetworkServer* InServer)
-	:	UbaServer(InServer)
+FUbaHordeAgentManager::FUbaHordeAgentManager(const FString& InWorkingDir, uba::NetworkServer* InServer)
+	:	WorkingDir(InWorkingDir)
+	,	UbaServer(InServer)
 	,	LastRequestFailTime(1)
 	,	TargetCoreCount(0)
 	,	EstimatedCoreCount(0)
@@ -188,10 +189,8 @@ void FUbaHordeAgentManager::ThreadAgent(FHordeAgentWrapper& Wrapper)
 		FScopeLock ScopeLock(&UbaAgentBundleFilePathLock);
 		if (UbaAgentBundleFilePath.IsEmpty())
 		{
-			FString RootWorkingDirectory = FPaths::Combine(FPlatformProcess::UserTempDir(), TEXT("UbaWorkingDir"));
-
 			const FString UbaAgentFilePath = FPaths::Combine(GetUbaBinariesPath(), TEXT("UbaAgent.exe"));
-			UbaAgentBundleFilePath = FPaths::Combine(RootWorkingDirectory, TEXT("UbaAgent.Bundle.ref"));
+			UbaAgentBundleFilePath = FPaths::Combine(WorkingDir, TEXT("UbaAgent.Bundle.ref"));
 
 			if (!CreateHordeBundleFromFile(*UbaAgentFilePath, *UbaAgentBundleFilePath))
 			{
@@ -250,6 +249,10 @@ void FUbaHordeAgentManager::ThreadAgent(FHordeAgentWrapper& Wrapper)
 		// If the machine couldn't be assigned, just ignore this agent slot
 		if (MachineInfo.Ip == TEXT(""))
 		{
+			if (!LastRequestFailTime)
+			{
+				UE_LOG(LogUbaHorde, Verbose, TEXT("No resources available in Horde. Will keep retrying until %u cores are used (Currently have %u)"), TargetCoreCount.Load(), ActiveCoreCount.Load());
+			}
 			LastRequestFailTime = FPlatformTime::Cycles64();
 			return;
 		}
@@ -336,6 +339,7 @@ void FUbaHordeAgentManager::ThreadAgent(FHordeAgentWrapper& Wrapper)
 
 		MachineCoreCount = MachineInfo.LogicalCores;
 		EstimatedCoreCount += MachineCoreCount;
+		ActiveCoreCount += MachineCoreCount;
 	}
 
 	while (Agent->IsValid() && !ShouldExit.Wait(100))
@@ -343,6 +347,7 @@ void FUbaHordeAgentManager::ThreadAgent(FHordeAgentWrapper& Wrapper)
 		Agent->Poll(UbaControllerModule::bHordeForwardAgentLogs);
 	}
 
+	ActiveCoreCount -= MachineCoreCount;
 	EstimatedCoreCount -= MachineCoreCount;
 }
 
diff --git a/Engine/Plugins/UbaController/Source/UbaController/Private/UbaHordeMetaClient.cpp b/Engine/Plugins/UbaController/Source/UbaController/Private/UbaHordeMetaClient.cpp
--- a/Engine/Plugins/UbaController/Source/UbaController/Private/UbaHordeMetaClient.cpp
+++ b/Engine/Plugins/UbaController/Source/UbaController/Private/UbaHordeMetaClient.cpp
@@ -78,7 +78,7 @@ TSharedPtr<FUbaHordeMetaClient::HordeMachinePromise, ESPMode::ThreadSafe> FUbaHo
 
 			if (HttpResponse->GetResponseCode() == 503)
 			{
-				UE_LOG(LogUbaHorde, Display, TEXT("No resources available in Horde (%s)"), *ResponseStr);
+				UE_LOG(LogUbaHorde, Verbose, TEXT("No resources available in Horde (%s)"), *ResponseStr);
 				Promise->SetValue(MakeTuple(HttpResponse, Info));
 				return;
 			}
diff --git a/Engine/Plugins/UbaController/Source/UbaController/Private/UbaJobProcessor.cpp b/Engine/Plugins/UbaController/Source/UbaController/Private/UbaJobProcessor.cpp
--- a/Engine/Plugins/UbaController/Source/UbaController/Private/UbaJobProcessor.cpp
+++ b/Engine/Plugins/UbaController/Source/UbaController/Private/UbaJobProcessor.cpp
@@ -60,6 +60,12 @@ namespace UbaJobProcessorOptions
 		bShowUbaLog,
 		TEXT("If true, UBA log entries will be visible in the log\n"));
 
+	static bool bProcessLogEnabled = false;
+	static FAutoConsoleVariableRef CVarProcessLogEnabled(
+		TEXT("r.UbaController.ProcessLogEnabled"),
+		bProcessLogEnabled,
+		TEXT("If true, each detoured process will write a log file. Note this is only useful if UBA is compiled in debug\n"));
+
 	FString ReplaceEnvironmentVariablesInPath(const FString& ExtraFilePartialPath) // Duplicated code with FAST build.. put it somewhere else?
 	{
 		FString ParsedPath;
@@ -109,8 +115,21 @@ FUbaJobProcessor::FUbaJobProcessor(
 	bIsWorkDone(false),
 	LogWriter([]() {}, []() {}, [](uba::LogEntryType type, const wchar_t* str, uba::u32 strlen)
 	{
-			if (UbaJobProcessorOptions::bShowUbaLog)
-				UE_LOG(LogUbaController, Log, TEXT("%s"), str);
+			switch (type)
+			{
+			case uba::LogEntryType_Error:
+				UE_LOG(LogUbaController, Error, TEXT("%s"), str);
+				break;
+			case uba::LogEntryType_Warning:
+				UE_LOG(LogUbaController, Warning, TEXT("%s"), str);
+				break;
+			default:
+				if (UbaJobProcessorOptions::bShowUbaLog)
+				{
+					UE_LOG(LogUbaController, Display, TEXT("%s"), str);
+				}
+				break;
+			}
 	})
 {
 	Uba_SetCustomAssertHandler([](const uba::tchar* text)
@@ -175,7 +194,10 @@ void FUbaJobProcessor::CalculateKnownInputs()
 
 		for (const FString& ExtraFilePartialPath : KnownFileNames)
 		{
-			AddKnownInput(UbaJobProcessorOptions::ReplaceEnvironmentVariablesInPath(ExtraFilePartialPath));
+			if (!ExtraFilePartialPath.Contains(TEXT("*"))) // Seems like there are some *.x paths in there.. TODO: Do a find files
+			{
+				AddKnownInput(UbaJobProcessorOptions::ReplaceEnvironmentVariablesInPath(ExtraFilePartialPath));
+			}
 		}
 	}
 
@@ -210,8 +232,8 @@ void FUbaJobProcessor::RunTaskWithUba(FTask* Task)
 	ProcessInfo.description = *InputFileName;
 	ProcessInfo.workingDir = *AppDir;
 
-	ProcessInfo.logFile = *LogPath;
-	ProcessInfo.trackInputs = false;
+	if (UbaJobProcessorOptions::bProcessLogEnabled)
+		ProcessInfo.logFile = *LogPath;
 	
 	struct ExitedInfo
 	{
@@ -247,17 +269,26 @@ void FUbaJobProcessor::StartUba()
 
 	UbaServer = CreateServer(LogWriter);
 
-	FString RootDir = FPaths::Combine(FPlatformProcess::UserTempDir(), TEXT("UnrealUbaTemp"));
+	FString RootDir = FString::Printf(TEXT("%s/%s%u"), FPlatformProcess::UserTempDir(), TEXT("UnrealUbaTemp"), UE::GetMultiprocessId());
 
-	UbaStorageServer = CreateStorageServer(*UbaServer, *RootDir, 0, true, LogWriter);
+	uba::u64 casCapacityBytes = 32llu * 1024 * 1024 * 1024;
+	UbaStorageServer = CreateStorageServer(*UbaServer, *RootDir, casCapacityBytes, true, LogWriter);
 
 	uba::SessionServerCreateInfo info(*UbaStorageServer, *UbaServer, LogWriter);
 	info.launchVisualizer = UbaJobProcessorOptions::bAutoLaunchVisualizer;
 	info.rootDir = *RootDir;
+	info.allowMemoryMaps = false; // Skip using memory maps
+
 	info.traceEnabled = true;
-	info.traceOutputFile = *UbaJobProcessorOptions::TraceFilename;
+	FString TraceOutputFile = UbaJobProcessorOptions::TraceFilename;
+	if (!TraceOutputFile.IsEmpty() && UE::GetMultiprocessId())
+		TraceOutputFile = FString::Printf(TEXT("%s_%u"), *TraceOutputFile, UE::GetMultiprocessId());
+	info.traceOutputFile = *TraceOutputFile;
 	info.detailedTrace = UbaJobProcessorOptions::bDetailedTrace;
-	info.allowMemoryMaps = false; // Skip using memory maps
+	FString TraceName = FString::Printf(TEXT("UbaController_%s"), *FGuid::NewGuid().ToString(EGuidFormats::Digits));
+	info.traceName = *TraceName;
+
+
 	//info.remoteLogEnabled = true;
 	UbaSessionServer = CreateSessionServer(info);
 
@@ -268,9 +299,12 @@ void FUbaJobProcessor::StartUba()
 
 	HandleTaskQueueUpdated(TEXT("")); // Flush tasks into uba scheduler
 
-	Server_StartListen(UbaServer, uba::DefaultPort, nullptr); // Start listen so any helper on the LAN can join in
+	if (UE::GetMultiprocessId() == 0)
+	{
+		Server_StartListen(UbaServer, uba::DefaultPort, nullptr); // Start listen so any helper on the LAN can join in
+	}
 
-	HordeAgentManager = MakeUnique<FUbaHordeAgentManager>(UbaServer);
+	HordeAgentManager = MakeUnique<FUbaHordeAgentManager>(ControllerModule.GetWorkingDirectory(), UbaServer);
 
 	UE_LOG(LogUbaController, Display, TEXT("Created UBA storage server: RootDir=%s"), *RootDir);
 }
diff --git a/Engine/Plugins/UbaController/Source/UbaController/Public/UbaHordeAgentManager.h b/Engine/Plugins/UbaController/Source/UbaController/Public/UbaHordeAgentManager.h
--- a/Engine/Plugins/UbaController/Source/UbaController/Public/UbaHordeAgentManager.h
+++ b/Engine/Plugins/UbaController/Source/UbaController/Public/UbaHordeAgentManager.h
@@ -11,7 +11,7 @@ namespace uba { class NetworkServer; }
 class FUbaHordeAgentManager
 {
 public:
-	FUbaHordeAgentManager(uba::NetworkServer* InServer);
+	FUbaHordeAgentManager(const FString& InWorkingDir, uba::NetworkServer* InServer);
 	~FUbaHordeAgentManager();
 
 	void SetTargetCoreCount(uint32 Count);
@@ -27,6 +27,7 @@ private:
 	void ThreadAgent(FHordeAgentWrapper& Wrapper);
 	void ParseConfig();
 
+	FString WorkingDir;
 	uba::NetworkServer* UbaServer;
 
 	FString Url;
@@ -45,5 +46,6 @@ private:
 	TAtomic<uint64> LastRequestFailTime;
 	TAtomic<uint32> TargetCoreCount;
 	TAtomic<uint32> EstimatedCoreCount;
+	TAtomic<uint32> ActiveCoreCount;
 	TAtomic<bool> AskForAgents;
 };
