[UBA/Shader]
* Changed Scheduler_GetStats to separate active into activeLocal and activeRemote
* Added Storage_DeleteFile that can be used to delete cas entries if we know a file will never be used again (no point taking up space)
* Fixed so ubaagents spawned through horde has a idle timeout of 15 seconds (also made them quiet)
* Changed so UbaJobProcessor shutdown uba after 100 seconds instead of 10 seconds of idling without tasks
* Added so .uba.in and .uba.out is deleted from cas storage after used.. no point storing them there since they will never be used again
* Fixed so number of local SCW instances are decreased when more remote helpers are running
* Fixed race condition where tasks could be added when uba was shutting down causing tasks to be cancelled

diff --git a/Engine/Plugins/UbaController/Source/UbaController/Private/UbaHordeAgentManager.cpp b/Engine/Plugins/UbaController/Source/UbaController/Private/UbaHordeAgentManager.cpp
--- a/Engine/Plugins/UbaController/Source/UbaController/Private/UbaHordeAgentManager.cpp
+++ b/Engine/Plugins/UbaController/Source/UbaController/Private/UbaHordeAgentManager.cpp
@@ -303,7 +303,9 @@ void FUbaHordeAgentManager::ThreadAgent(FHordeAgentWrapper& Wrapper)
 		{
 			ListenPortArg.c_str(),
 			"-nopoll",				// -nopoll recommended when running on remote Horde agents to make sure they exit after completion. Otherwise, it keeps running.
-			"-listenTimeout=20",	// increase timeout for the agent to listen from 5 to 20 seconds to give the editor enough time to establish the connection
+			"-listenTimeout=5",		// Agent will wait 5 seconds for this thread to connect (Server_AddClient does the connect)
+			"-quiet",				// Skip all the agent logging that would be sent over to here
+			"-maxidle=15",			// After 15 seconds of idling agent will automatically disconnect
 		};
 
 		// If the machine does not run Windows, enable the compatibility layer Wine to run UbaAgent.exe on POSIX systems
diff --git a/Engine/Plugins/UbaController/Source/UbaController/Private/UbaJobProcessor.cpp b/Engine/Plugins/UbaController/Source/UbaController/Private/UbaJobProcessor.cpp
--- a/Engine/Plugins/UbaController/Source/UbaController/Private/UbaJobProcessor.cpp
+++ b/Engine/Plugins/UbaController/Source/UbaController/Private/UbaJobProcessor.cpp
@@ -24,7 +24,7 @@ namespace UbaJobProcessorOptions
         SleepTimeBetweenActions,
         TEXT("How much time the job processor thread should sleep between actions .\n"));
 
-	static float MaxTimeWithoutTasks = 10.0f;
+	static float MaxTimeWithoutTasks = 100.0f;
 	static FAutoConsoleVariableRef CVarMaxTimeWithoutTasks(
         TEXT("r.UbaController.MaxTimeWithoutTasks"),
         MaxTimeWithoutTasks,
@@ -224,7 +224,7 @@ void FUbaJobProcessor::RunTaskWithUba(FTask* Task)
 	FString OutputFileName = FPaths::GetCleanFilename(Data.OutputFileName);
 	FString Parameters = FString::Printf(TEXT("\"%s/\" %d 0 \"%s\" \"%s\" %s "), *Data.WorkingDirectory, Data.DispatcherPID, *InputFileName, *OutputFileName, *Data.ExtraCommandArgs);
 	FString AppDir = FPaths::GetPath(Data.Command);
-	FString LogPath = FString::FromInt(Task->ID) + TEXT(".Log");
+	FString LogPath;
 
 	uba::ProcessStartInfo ProcessInfo;
 	ProcessInfo.application = *Data.Command;
@@ -234,18 +234,23 @@ void FUbaJobProcessor::RunTaskWithUba(FTask* Task)
 	ProcessInfo.writeOutputFilesOnFail = true;
 
 	if (UbaJobProcessorOptions::bProcessLogEnabled)
+	{
+		LogPath = FString::FromInt(Task->ID) + TEXT(".Log");
 		ProcessInfo.logFile = *LogPath;
+	}
 	
 	struct ExitedInfo
 	{
 		FUbaJobProcessor* Processor;
 		FString InputFile;
+		FString OutputFile;
 		FTask* Task;
 	};
 
 	auto Info = new ExitedInfo;
 	Info->Processor = this;
 	Info->InputFile = Data.InputFileName;
+	Info->OutputFile = Data.OutputFileName;
 	Info->Task = Task;
 
 	ProcessInfo.userData = Info;
@@ -256,6 +261,10 @@ void FUbaJobProcessor::RunTaskWithUba(FTask* Task)
 				IFileManager::Get().Delete(*Info->InputFile);
 				SessionServer_RegisterDeleteFile(Info->Processor->UbaSessionServer, *Info->InputFile);
 				Info->Processor->HandleUbaJobFinished(Info->Task);
+
+				Storage_DeleteFile(Info->Processor->UbaStorageServer, *Info->InputFile);
+				Storage_DeleteFile(Info->Processor->UbaStorageServer, *Info->OutputFile);
+
 				delete Info;
 			}
 		};
@@ -347,12 +356,17 @@ uint32 FUbaJobProcessor::Run()
 		const float ElapsedSeconds = (FPlatformTime::Cycles() - LastTimeSinceHadJobs) * FPlatformTime::GetSecondsPerCycle();
 
 		uint32 queued = 0;
-		uint32 active = 0;
+		uint32 activeLocal = 0;
+		uint32 activeRemote = 0;
 		uint32 finished = 0;
+
+		FScopeLock lock(&bShouldProcessJobsLock);
+
 		if (UbaScheduler)
 		{
-			Scheduler_GetStats(UbaScheduler, queued, active, finished);
+			Scheduler_GetStats(UbaScheduler, queued, activeLocal, activeRemote, finished);
 		}
+		uint32 active = activeLocal + activeRemote;
 
 		// We don't want to hog up Horde resources.
 		if (bShouldProcessJobs && ElapsedSeconds > UbaJobProcessorOptions::MaxTimeWithoutTasks && (queued + active) == 0)
@@ -378,6 +392,9 @@ uint32 FUbaJobProcessor::Run()
 
 		if (bShouldProcessJobs)
 		{
+			int32 MaxLocal = FMath::Max(1, int32(MaxLocalParallelJobs / 2) - int32(activeRemote / 5));
+			Scheduler_SetMaxLocalProcessors(UbaScheduler, MaxLocal);
+
 			int32 TargetCoreCount = FMath::Max(0, int32(queued + active) - MaxLocalParallelJobs);
 
 			HordeAgentManager->SetTargetCoreCount(TargetCoreCount);
@@ -386,6 +403,8 @@ uint32 FUbaJobProcessor::Run()
 			SessionServer_SetMaxRemoteProcessCount(UbaSessionServer, TargetCoreCount);
 		}
 
+		lock.Unlock();
+
 		FPlatformProcess::Sleep(UbaJobProcessorOptions::SleepTimeBetweenActions);
 	}
 
@@ -469,6 +488,8 @@ void FUbaJobProcessor::HandleUbaJobFinished(FTask* CompileTask)
 
 void FUbaJobProcessor::HandleTaskQueueUpdated(const FString& InputFileName)
 {
+	FScopeLock lock(&bShouldProcessJobsLock);
+
 	if (!UbaScheduler)
 	{
 		return;
@@ -490,9 +511,9 @@ bool FUbaJobProcessor::HasJobsInFlight() const
 		return false;
 	}
 	uint32 queued = 0;
-	uint32 active = 0;
+	uint32 activeLocal = 0;
+	uint32 activeRemote = 0;
 	uint32 finished = 0;
-	Scheduler_GetStats(UbaScheduler, queued, active, finished);
-		
-	return queued != 0 || active != 0;
+	Scheduler_GetStats(UbaScheduler, queued, activeLocal, activeRemote, finished);	
+	return (queued + activeLocal + activeRemote) != 0;
 }
\ No newline at end of file
diff --git a/Engine/Plugins/UbaController/Source/UbaController/Public/UbaJobProcessor.h b/Engine/Plugins/UbaController/Source/UbaController/Public/UbaJobProcessor.h
--- a/Engine/Plugins/UbaController/Source/UbaController/Public/UbaJobProcessor.h
+++ b/Engine/Plugins/UbaController/Source/UbaController/Public/UbaJobProcessor.h
@@ -59,7 +59,8 @@ private:
 	uint32 LastTimeCheckedForTasks;
 	
 	/** Used to stop the processing loop in the next main loop */
-	TAtomic<bool> bShouldProcessJobs;
+	FCriticalSection bShouldProcessJobsLock;
+	bool bShouldProcessJobs;
 
 	/** Set to true when the main loop finishes*/
 	bool bIsWorkDone;
diff --git a/Engine/Source/Programs/UnrealBuildAccelerator/Common/Private/UbaFile.cpp b/Engine/Source/Programs/UnrealBuildAccelerator/Common/Private/UbaFile.cpp
--- a/Engine/Source/Programs/UnrealBuildAccelerator/Common/Private/UbaFile.cpp
+++ b/Engine/Source/Programs/UnrealBuildAccelerator/Common/Private/UbaFile.cpp
@@ -187,7 +187,11 @@ namespace uba
 		struct stat attr;
 		int res = stat(fileName, &attr);
 		if (res != 0)
-			return false;// logger.Error(TC("GetFileInformation: CreateFile failed for file %s (%s)"), fileName, strerror(errno));
+		{
+			UBA_ASSERTF(errno == ENOENT, TC("GetFileInformation: stat failed for file %s and this error handling not implemented (%s)"), fileName, strerror(errno));
+			SetLastError(ERROR_FILE_NOT_FOUND);
+			return false;
+		}
 		out.lastWriteTime = FromTimeSpec(attr.st_mtimespec);
 		out.attributes = attr.st_mode;
 		out.volumeSerialNumber = attr.st_dev;
diff --git a/Engine/Source/Programs/UnrealBuildAccelerator/Common/Private/UbaScheduler.cpp b/Engine/Source/Programs/UnrealBuildAccelerator/Common/Private/UbaScheduler.cpp
--- a/Engine/Source/Programs/UnrealBuildAccelerator/Common/Private/UbaScheduler.cpp
+++ b/Engine/Source/Programs/UnrealBuildAccelerator/Common/Private/UbaScheduler.cpp
@@ -95,10 +95,11 @@ namespace uba
 		return 0;
 	}
 
-	void Scheduler::GetStats(u32& outQueued, u32& outActive, u32& outFinished)
+	void Scheduler::GetStats(u32& outQueued, u32& outActiveLocal, u32& outActiveRemote, u32& outFinished)
 	{
 		ScopedReadLock lock(m_queuedProcessesLock);
-		outActive = m_activeLocalProcesses + m_activeRemoteProcesses;
+		outActiveLocal = m_activeLocalProcesses;
+		outActiveRemote = m_activeRemoteProcesses;
 		outFinished = m_finishedProcesses;
 		outQueued = u32(m_queuedProcesses.size());
 	}
diff --git a/Engine/Source/Programs/UnrealBuildAccelerator/Common/Private/UbaStorage.cpp b/Engine/Source/Programs/UnrealBuildAccelerator/Common/Private/UbaStorage.cpp
--- a/Engine/Source/Programs/UnrealBuildAccelerator/Common/Private/UbaStorage.cpp
+++ b/Engine/Source/Programs/UnrealBuildAccelerator/Common/Private/UbaStorage.cpp
@@ -2383,4 +2383,26 @@ namespace uba
 	{
 		return m_dirCache.CreateDirectory(m_logger, dir);
 	}
+
+	bool StorageImpl::DeleteCasForFile(const tchar* file)
+	{
+		StringBuffer<> forKey;
+		FixPath(file, nullptr, 0, forKey);
+		if (CaseInsensitiveFs)
+			forKey.MakeLower();
+		StringKey fileNameKey = ToStringKey(forKey);
+
+		ScopedReadLock lookupLock(m_fileTableLookupLock);
+		auto findIt = m_fileTableLookup.find(fileNameKey);
+		if (findIt == m_fileTableLookup.end())
+			return false;
+		FileEntry& fileEntry = findIt->second;
+		lookupLock.Leave();
+
+		ScopedWriteLock entryLock(fileEntry.lock);
+		fileEntry.verified = false;
+
+		return DropCasFile(fileEntry.casKey, true);
+	}
+
 }
diff --git a/Engine/Source/Programs/UnrealBuildAccelerator/Common/Public/UbaScheduler.h b/Engine/Source/Programs/UnrealBuildAccelerator/Common/Public/UbaScheduler.h
--- a/Engine/Source/Programs/UnrealBuildAccelerator/Common/Public/UbaScheduler.h
+++ b/Engine/Source/Programs/UnrealBuildAccelerator/Common/Public/UbaScheduler.h
@@ -27,7 +27,7 @@ namespace uba
 
 		u32 EnqueueProcess(const ProcessStartInfo& info, float weight = 1.0f, const void* knownInputs = nullptr, u32 knownInputsBytes = 0, u32 knownInputsCount = 0);
 
-		void GetStats(u32& outQueued, u32& outActive, u32& outFinished);
+		void GetStats(u32& outQueued, u32& outActiveLocal, u32& outActiveRemote, u32& outFinished);
 	private:
 		struct ExitProcessInfo;
 
diff --git a/Engine/Source/Programs/UnrealBuildAccelerator/Common/Public/UbaStorage.h b/Engine/Source/Programs/UnrealBuildAccelerator/Common/Public/UbaStorage.h
--- a/Engine/Source/Programs/UnrealBuildAccelerator/Common/Public/UbaStorage.h
+++ b/Engine/Source/Programs/UnrealBuildAccelerator/Common/Public/UbaStorage.h
@@ -35,6 +35,7 @@ namespace uba
 		virtual bool DecompressFileToMemory(const tchar* fileName, FileHandle fileHandle, u8* dest, u64 decompressedSize) = 0;
 		virtual bool DecompressMemoryToMemory(u8* compressedData, u8* writeData, u64 decompressedSize) = 0;
 		virtual bool CreateDirectory(const tchar* dir) = 0;
+		virtual bool DeleteCasForFile(const tchar* file) = 0;
 
 		struct RetrieveResult { CasKey casKey; u64 size = 0; MappedView view; };
 		virtual bool RetrieveCasFile(RetrieveResult& out, const CasKey& casKey, const tchar* hint, FileMappingBuffer* mappingBuffer = nullptr, u64 memoryMapAlignment = 1, bool allowProxy = true) = 0;
@@ -96,6 +97,7 @@ namespace uba
 
 		virtual bool DecompressFileToMemory(const tchar* fileName, FileHandle fileHandle, u8* dest, u64 decompressedSize) override;
 		virtual bool CreateDirectory(const tchar* dir) override;
+		virtual bool DeleteCasForFile(const tchar* file) override;
 		virtual bool RetrieveCasFile(RetrieveResult& out, const CasKey& casKey, const tchar* hint, FileMappingBuffer* mappingBuffer = nullptr, u64 memoryMapAlignment = 1, bool allowProxy = true) override;
 		virtual bool VerifyAndGetCachedFileInfo(CachedFileInfo& out, StringKey fileNameKey, u64 verifiedLastWriteTime, u64 verifiedSize) override;
 		virtual bool StoreCasFile(CasKey& out, const tchar* fileName, const CasKey& casKeyOverride = CasKeyZero, bool deferCreation = false) override;
diff --git a/Engine/Source/Programs/UnrealBuildAccelerator/Host/Private/UbaExports.cpp b/Engine/Source/Programs/UnrealBuildAccelerator/Host/Private/UbaExports.cpp
--- a/Engine/Source/Programs/UnrealBuildAccelerator/Host/Private/UbaExports.cpp
+++ b/Engine/Source/Programs/UnrealBuildAccelerator/Host/Private/UbaExports.cpp
@@ -193,6 +193,11 @@ extern "C"
 		storage->SaveCasTable(true);
 	}
 
+	void Storage_DeleteFile(uba::Storage* storage, const uba::tchar* file)
+	{
+		storage->DeleteCasForFile(file);
+	}
+
 	uba::Storage* CreateStorageServer(uba::NetworkServer& server, const uba::tchar* rootDir, uba::u64 casCapacityBytes, bool storeCompressed, uba::LogWriter& writer, const uba::tchar* zone)
 	{
 		using namespace uba;
@@ -471,9 +476,9 @@ uba::StorageClient* CreateStorageClient(uba::NetworkClient& client, const uba::t
 		delete scheduler;
 	}
 
-	void Scheduler_GetStats(uba::Scheduler* scheduler, uba::u32& outQueued, uba::u32& outActive, uba::u32& outFinished)
+	void Scheduler_GetStats(uba::Scheduler* scheduler, uba::u32& outQueued, uba::u32& outActiveLocal, uba::u32& outActiveRemote, uba::u32& outFinished)
 	{
-		scheduler->GetStats(outQueued, outActive, outFinished);
+		scheduler->GetStats(outQueued, outActiveLocal, outActiveRemote, outFinished);
 	}
 
 	void Uba_SetCustomAssertHandler(Uba_CustomAssertHandler* handler)
diff --git a/Engine/Source/Programs/UnrealBuildAccelerator/Host/Public/UbaExports.h b/Engine/Source/Programs/UnrealBuildAccelerator/Host/Public/UbaExports.h
--- a/Engine/Source/Programs/UnrealBuildAccelerator/Host/Public/UbaExports.h
+++ b/Engine/Source/Programs/UnrealBuildAccelerator/Host/Public/UbaExports.h
@@ -55,6 +55,7 @@ extern "C"
 	UBA_API uba::Storage* CreateStorage(const uba::tchar* rootDir, uba::u64 casCapacityBytes, bool storeCompressed);
 	UBA_API void DestroyStorage(uba::Storage* storage);
 	UBA_API void Storage_SaveCasTable(uba::Storage* storage);
+	UBA_API void Storage_DeleteFile(uba::Storage* storage, const uba::tchar* file);
 
 	UBA_API uba::Storage* CreateStorageServer(uba::NetworkServer& server, const uba::tchar* rootDir, uba::u64 casCapacityBytes, bool storeCompressed, uba::LogWriter& writer = uba::g_consoleLogWriter, const uba::tchar* zone = TC(""));
 	UBA_API void DestroyStorageServer(uba::Storage* storageServer);
@@ -117,7 +118,7 @@ extern "C"
 	UBA_API void Scheduler_SetMaxLocalProcessors(uba::Scheduler* scheduler, uba::u32 maxLocalProcessors);
 	UBA_API void Scheduler_Stop(uba::Scheduler* scheduler);
 	UBA_API void Scheduler_Destroy(uba::Scheduler* scheduler);
-	UBA_API void Scheduler_GetStats(uba::Scheduler* scheduler, uba::u32& outQueued, uba::u32& outActive, uba::u32& outFinished);
+	UBA_API void Scheduler_GetStats(uba::Scheduler* scheduler, uba::u32& outQueued, uba::u32& outActiveLocal, uba::u32& outActiveRemote, uba::u32& outFinished);
 
 	// Misc
 	using Uba_CustomAssertHandler = void(const uba::tchar* text);
diff --git a/Engine/Source/Programs/UnrealBuildAccelerator/Test/Private/UbaTestScheduler.h b/Engine/Source/Programs/UnrealBuildAccelerator/Test/Private/UbaTestScheduler.h
--- a/Engine/Source/Programs/UnrealBuildAccelerator/Test/Private/UbaTestScheduler.h
+++ b/Engine/Source/Programs/UnrealBuildAccelerator/Test/Private/UbaTestScheduler.h
@@ -22,8 +22,8 @@ namespace uba
 				scheduler.EnqueueProcess(processInfo);
 				scheduler.Start();
 
-				u32 queued, active, finished;
-				do { scheduler.GetStats(queued, active, finished); } while (queued != 0 || active != 0);
+				u32 queued, activeLocal, activeRemote, finished;
+				do { scheduler.GetStats(queued, activeLocal, activeRemote, finished); } while ((queued+activeLocal+activeRemote) != 0);
 
 				scheduler.Stop();
 				return true;
@@ -51,8 +51,8 @@ namespace uba
 				scheduler.EnqueueProcess(processInfo);
 				scheduler.Start();
 
-				u32 queued, active, finished;
-				do { scheduler.GetStats(queued, active, finished); } while (queued != 0 || active != 0);
+				u32 queued, activeLocal, activeRemote, finished;
+				do { scheduler.GetStats(queued, activeLocal, activeRemote, finished); } while ((queued + activeLocal + activeRemote) != 0);
 
 				scheduler.Stop();
 				return true;
@@ -80,8 +80,8 @@ namespace uba
 				scheduler.EnqueueProcess(processInfo);
 				scheduler.Start();
 
-				u32 queued, active, finished;
-				do { scheduler.GetStats(queued, active, finished); } while (queued != 0 || active != 0);
+				u32 queued, activeLocal, activeRemote, finished;
+				do { scheduler.GetStats(queued, activeLocal, activeRemote, finished); } while ((queued + activeLocal + activeRemote) != 0);
 
 				scheduler.Stop();
 				return true;
