[UbaController]
* Moved temp directories to make it easier to find/delete them

diff --git a/Engine/Plugins/UbaController/Source/UbaController/Private/UbaControllerModule.cpp b/Engine/Plugins/UbaController/Source/UbaController/Private/UbaControllerModule.cpp
--- a/Engine/Plugins/UbaController/Source/UbaController/Private/UbaControllerModule.cpp
+++ b/Engine/Plugins/UbaController/Source/UbaController/Private/UbaControllerModule.cpp
@@ -24,7 +24,7 @@ FUbaControllerModule::FUbaControllerModule()
 	: bSupported(false)
 	, bModuleInitialized(false)
 	, bControllerInitialized(false)
-	, RootWorkingDirectory(FPaths::Combine(FPlatformProcess::UserTempDir(), TEXT("UbaWorkingDir")))
+	, RootWorkingDirectory(FPaths::Combine(FPlatformProcess::UserTempDir(), TEXT("UbaControllerWorkingDir")))
 	, WorkingDirectory(FPaths::Combine(RootWorkingDirectory, FGuid::NewGuid().ToString(EGuidFormats::Digits)))
 	, NextFileID(0)
 	, NextTaskID(0)
diff --git a/Engine/Plugins/UbaController/Source/UbaController/Private/UbaJobProcessor.cpp b/Engine/Plugins/UbaController/Source/UbaController/Private/UbaJobProcessor.cpp
--- a/Engine/Plugins/UbaController/Source/UbaController/Private/UbaJobProcessor.cpp
+++ b/Engine/Plugins/UbaController/Source/UbaController/Private/UbaJobProcessor.cpp
@@ -279,7 +279,8 @@ void FUbaJobProcessor::StartUba()
 
 	UbaServer = CreateServer(LogWriter);
 
-	FString RootDir = FString::Printf(TEXT("%s/%s%u"), FPlatformProcess::UserTempDir(), TEXT("UnrealUbaTemp"), UE::GetMultiprocessId());
+	FString RootDir = FString::Printf(TEXT("%s/%s/%u"), FPlatformProcess::UserTempDir(), TEXT("UbaControllerStorageDir"), UE::GetMultiprocessId());
+	IFileManager::Get().MakeDirectory(*RootDir, true);
 
 	uba::u64 casCapacityBytes = 32llu * 1024 * 1024 * 1024;
 	UbaStorageServer = CreateStorageServer(*UbaServer, *RootDir, casCapacityBytes, true, LogWriter);
@@ -392,10 +393,10 @@ uint32 FUbaJobProcessor::Run()
 
 		if (bShouldProcessJobs)
 		{
-			int32 MaxLocal = FMath::Max(1, int32(MaxLocalParallelJobs / 2) - int32(activeRemote / 5));
+			int32 MaxLocal = FMath::Max(1, int32(MaxLocalParallelJobs / 2) - int32(activeRemote / 10));
 			Scheduler_SetMaxLocalProcessors(UbaScheduler, MaxLocal);
 
-			int32 TargetCoreCount = FMath::Max(0, int32(queued + active) - MaxLocalParallelJobs);
+			int32 TargetCoreCount = FMath::Max(0, int32(queued + active) - MaxLocal);
 
 			HordeAgentManager->SetTargetCoreCount(TargetCoreCount);
 			
