[UBA]
* Improved logging with the mode where log files are sent back to host from remotes

diff --git a/Engine/Plugins/UbaController/Source/UbaController/Private/UbaJobProcessor.cpp b/Engine/Plugins/UbaController/Source/UbaController/Private/UbaJobProcessor.cpp
--- a/Engine/Plugins/UbaController/Source/UbaController/Private/UbaJobProcessor.cpp
+++ b/Engine/Plugins/UbaController/Source/UbaController/Private/UbaJobProcessor.cpp
@@ -224,7 +224,6 @@ void FUbaJobProcessor::RunTaskWithUba(FTask* Task)
 	FString OutputFileName = FPaths::GetCleanFilename(Data.OutputFileName);
 	FString Parameters = FString::Printf(TEXT("\"%s/\" %d 0 \"%s\" \"%s\" %s "), *Data.WorkingDirectory, Data.DispatcherPID, *InputFileName, *OutputFileName, *Data.ExtraCommandArgs);
 	FString AppDir = FPaths::GetPath(Data.Command);
-	FString LogPath;
 
 	uba::ProcessStartInfo ProcessInfo;
 	ProcessInfo.application = *Data.Command;
@@ -235,8 +234,7 @@ void FUbaJobProcessor::RunTaskWithUba(FTask* Task)
 
 	if (UbaJobProcessorOptions::bProcessLogEnabled)
 	{
-		LogPath = FString::FromInt(Task->ID) + TEXT(".Log");
-		ProcessInfo.logFile = *LogPath;
+		ProcessInfo.logFile = *InputFileName;
 	}
 	
 	struct ExitedInfo
@@ -254,8 +252,14 @@ void FUbaJobProcessor::RunTaskWithUba(FTask* Task)
 	Info->Task = Task;
 
 	ProcessInfo.userData = Info;
-	ProcessInfo.exitedFunc = [](void* userData, const uba::ProcessHandle&)
+	ProcessInfo.exitedFunc = [](void* userData, const uba::ProcessHandle& ph)
 		{
+			uint32 logLineIndex = 0;
+			while (const uba::tchar* line = ProcessHandle_GetLogLine(&ph, logLineIndex++))
+			{
+				UE_LOG(LogUbaController, Display, TEXT("%s"), line);
+			}
+
 			if (auto Info = (ExitedInfo*)userData) // It can be null if custom message has already handled all of them
 			{
 				IFileManager::Get().Delete(*Info->InputFile);
@@ -289,6 +293,7 @@ void FUbaJobProcessor::StartUba()
 	info.launchVisualizer = UbaJobProcessorOptions::bAutoLaunchVisualizer;
 	info.rootDir = *RootDir;
 	info.allowMemoryMaps = false; // Skip using memory maps
+	info.remoteLogEnabled = UbaJobProcessorOptions::bProcessLogEnabled;
 
 	info.traceEnabled = true;
 	FString TraceOutputFile = UbaJobProcessorOptions::TraceFilename;
@@ -393,7 +398,7 @@ uint32 FUbaJobProcessor::Run()
 
 		if (bShouldProcessJobs)
 		{
-			int32 MaxLocal = FMath::Max(1, int32(MaxLocalParallelJobs / 2) - int32(activeRemote / 10));
+			int32 MaxLocal = FMath::Max(0, int32(MaxLocalParallelJobs / 2) - int32(activeRemote / 10));
 			Scheduler_SetMaxLocalProcessors(UbaScheduler, MaxLocal);
 
 			int32 TargetCoreCount = FMath::Max(0, int32(queued + active) - MaxLocal);
