Provide lease ID in UBA/Horde URL to forward user directly to the job in the web portal.

#rnx

diff --git a/Engine/Plugins/UbaController/Source/UbaController/Private/UbaHordeMetaClient.cpp b/Engine/Plugins/UbaController/Source/UbaController/Private/UbaHordeMetaClient.cpp
--- a/Engine/Plugins/UbaController/Source/UbaController/Private/UbaHordeMetaClient.cpp
+++ b/Engine/Plugins/UbaController/Source/UbaController/Private/UbaHordeMetaClient.cpp
@@ -122,13 +122,13 @@ TSharedPtr<FUbaHordeMetaClient::HordeMachinePromise, ESPMode::ThreadSafe> FUbaHo
 				}
 			}
 
-			FString NonceString = NonceValue->AsString(); 
+			FString NonceString = NonceValue->AsString();
 			FString IpString = IpValue->AsString();
 			uint16 PortNumber = (uint16)PortValue->AsNumber();
 
-			if (TSharedPtr<FJsonValue> AgentIdValue = OutJson->AsObject()->TryGetField(TEXT("agentId")))
+			if (TSharedPtr<FJsonValue> LeaseIdValue = OutJson->AsObject()->TryGetField(TEXT("leaseId")))
 			{
-				const FString AgentWebPortalUrl = FString::Format(TEXT("{0}/agents?agentId={1}"), { this->ServerUrl, AgentIdValue->AsString() });
+				const FString AgentWebPortalUrl = FString::Format(TEXT("{0}/lease/{1}"), { this->ServerUrl, LeaseIdValue->AsString() });
 				UE_LOG(
 					LogUbaHorde, Display, TEXT("UBA Horde machine assigned (%s) [%s:%u]: %s"),
 					(Info.bRunsWindowOS ? TEXT("Microsoft Windows") : TEXT("POSIX/Wine")), *IpString, (uint32)PortNumber, *AgentWebPortalUrl
