Use UE_HORDE_TOKEN environment variable instead of OIDC login when unattended due to lack of user interaction on farm builds.

#rnx

diff --git a/Engine/Plugins/UbaController/Source/UbaController/Private/UbaHordeMetaClient.cpp b/Engine/Plugins/UbaController/Source/UbaController/Private/UbaHordeMetaClient.cpp
--- a/Engine/Plugins/UbaController/Source/UbaController/Private/UbaHordeMetaClient.cpp
+++ b/Engine/Plugins/UbaController/Source/UbaController/Private/UbaHordeMetaClient.cpp
@@ -30,11 +30,23 @@ bool FUbaHordeMetaClient::RefreshHttpClient()
 	if (bConnectWithAuthentication)
 	{
 		check(!OAuthProviderIdentifier.IsEmpty());
-		UE_LOG(LogUbaHorde, Display, TEXT("Logging in to Horde server with OIDC: %s"), *ServerUrl);
-		if (!HttpClient->LoginWithOidc(*OAuthProviderIdentifier, FApp::IsUnattended()))
+		if (FApp::IsUnattended())
 		{
-			UE_LOG(LogUbaHorde, Error, TEXT("Login to Horde server [%s] failed"), *ServerUrl);
-			return false;
+			UE_LOG(LogUbaHorde, Display, TEXT("Logging in to Horde server with environment variable UE_HORDE_TOKEN: %s"), *ServerUrl);
+			if (!HttpClient->LoginWithEnvironmentVariable())
+			{
+				UE_LOG(LogUbaHorde, Error, TEXT("Login to Horde server [%s] failed"), *ServerUrl);
+				return false;
+			}
+		}
+		else
+		{
+			UE_LOG(LogUbaHorde, Display, TEXT("Logging in to Horde server with OIDC: %s"), *ServerUrl);
+			if (!HttpClient->LoginWithOidc(*OAuthProviderIdentifier, FApp::IsUnattended()))
+			{
+				UE_LOG(LogUbaHorde, Error, TEXT("Login to Horde server [%s] failed"), *ServerUrl);
+				return false;
+			}
 		}
 	}
 	return true;
diff --git a/Engine/Source/Developer/Horde/Private/HordeHttpClient.cpp b/Engine/Source/Developer/Horde/Private/HordeHttpClient.cpp
--- a/Engine/Source/Developer/Horde/Private/HordeHttpClient.cpp
+++ b/Engine/Source/Developer/Horde/Private/HordeHttpClient.cpp
@@ -7,6 +7,7 @@
 #include "IDesktopPlatform.h"
 #include "HttpModule.h"
 #include "Misc/Paths.h"
+#include "HAL/PlatformMisc.h"
 
 FHordeHttpClient::FHordeHttpClient(FString InServerUrl)
 	: ServerUrl(InServerUrl)
@@ -35,6 +36,17 @@ bool FHordeHttpClient::LoginWithOidc(const TCHAR* Profile, bool bUnattended, FFe
 	return false;
 }
 
+bool FHordeHttpClient::LoginWithEnvironmentVariable()
+{
+	FString EnvVarHordeToken = FPlatformMisc::GetEnvironmentVariable(TEXT("UE_HORDE_TOKEN"));
+	if (!EnvVarHordeToken.IsEmpty())
+	{
+		Token = MoveTemp(EnvVarHordeToken);
+		return true;
+	}
+	return false;
+}
+
 TSharedRef<IHttpRequest> FHordeHttpClient::CreateRequest(const TCHAR* Verb, const TCHAR* Path)
 {
 	FHttpModule& Module = FHttpModule::Get();
diff --git a/Engine/Source/Developer/Horde/Public/HordeHttpClient.h b/Engine/Source/Developer/Horde/Public/HordeHttpClient.h
--- a/Engine/Source/Developer/Horde/Public/HordeHttpClient.h
+++ b/Engine/Source/Developer/Horde/Public/HordeHttpClient.h
@@ -14,6 +14,8 @@ public:
 	~FHordeHttpClient();
 
 	bool LoginWithOidc(const TCHAR* Profile, bool bUnattended, FFeedbackContext* Warn = nullptr);
+
+	bool LoginWithEnvironmentVariable();
 	
 	TSharedRef<IHttpRequest> CreateRequest(const TCHAR* Verb, const TCHAR* Path);
 	TSharedRef<IHttpResponse> Get(const TCHAR* Path);
