Replace ensure condition with return code to allow UBA agents to fail on volatile cloud services.

#rnx

diff --git a/Engine/Plugins/UbaController/Source/UbaController/Private/UbaHordeAgent.cpp b/Engine/Plugins/UbaController/Source/UbaController/Private/UbaHordeAgent.cpp
--- a/Engine/Plugins/UbaController/Source/UbaController/Private/UbaHordeAgent.cpp
+++ b/Engine/Plugins/UbaController/Source/UbaController/Private/UbaHordeAgent.cpp
@@ -124,7 +124,7 @@ bool FUbaHordeAgent::UploadBinaries(const FString& BundleDirectory, const char*
 			}
 		};
 
-	EAgentMessageType Type;
+	EAgentMessageType Type = EAgentMessageType::None;
 
 	for (;;)
 	{
@@ -157,12 +157,11 @@ bool FUbaHordeAgent::UploadBinaries(const FString& BundleDirectory, const char*
 		ChildChannel->Blob((unsigned char*)SerializedBytes.GetData(), SerializedBytes.Num());
 	}
 
-	if (!ensure(Type == EAgentMessageType::WriteFilesResponse))
-	{
-		return false;
-	}
+	// In case the agent upload was successful, WriteFilesResponse must have been the last received response.
+	// Otherwise, the remote machine might have abruptly shutdown, which can happen quite frequently on certain cloud services.
+	const bool bSuccess = (Type == EAgentMessageType::WriteFilesResponse);
 
-	return true;
+	return bSuccess;
 }
 
 void FUbaHordeAgent::Execute(const char* Exe, const char** Args, size_t NumArgs, const char* WorkingDir, const char** EnvVars, size_t NumEnvVars, bool bUseWine)
