Change verbosity of failed Horde login attempt for UBA agents as these failures must be accepted.

#rnx

diff --git a/Engine/Plugins/UbaController/Source/UbaController/Private/UbaHordeMetaClient.cpp b/Engine/Plugins/UbaController/Source/UbaController/Private/UbaHordeMetaClient.cpp
--- a/Engine/Plugins/UbaController/Source/UbaController/Private/UbaHordeMetaClient.cpp
+++ b/Engine/Plugins/UbaController/Source/UbaController/Private/UbaHordeMetaClient.cpp
@@ -35,7 +35,7 @@ bool FUbaHordeMetaClient::RefreshHttpClient()
 			UE_LOG(LogUbaHorde, Display, TEXT("Logging in to Horde server with environment variable UE_HORDE_TOKEN: %s"), *ServerUrl);
 			if (!HttpClient->LoginWithEnvironmentVariable())
 			{
-				UE_LOG(LogUbaHorde, Error, TEXT("Login to Horde server [%s] failed"), *ServerUrl);
+				UE_LOG(LogUbaHorde, Warning, TEXT("Login to Horde server [%s] failed"), *ServerUrl);
 				return false;
 			}
 		}
@@ -44,7 +44,7 @@ bool FUbaHordeMetaClient::RefreshHttpClient()
 			UE_LOG(LogUbaHorde, Display, TEXT("Logging in to Horde server with OIDC: %s"), *ServerUrl);
 			if (!HttpClient->LoginWithOidc(*OAuthProviderIdentifier, FApp::IsUnattended()))
 			{
-				UE_LOG(LogUbaHorde, Error, TEXT("Login to Horde server [%s] failed"), *ServerUrl);
+				UE_LOG(LogUbaHorde, Warning, TEXT("Login to Horde server [%s] failed"), *ServerUrl);
 				return false;
 			}
 		}
@@ -81,7 +81,7 @@ TSharedPtr<FUbaHordeMetaClient::HordeMachinePromise, ESPMode::ThreadSafe> FUbaHo
 
 			if (!bSucceeded || !HttpResponse.IsValid())
 			{
-				UE_LOG(LogUbaHorde, Warning, TEXT("No response from Horde"));
+				UE_LOG(LogUbaHorde, Verbose, TEXT("No response from Horde"));
 				Promise->SetValue(MakeTuple(HttpResponse, Info));
 				return;
 			}
