[UBA] Fall back to OIDC login even for unattended cooks when environment variable UE_HORDE_TOKEN is unavailable.

#rnx

diff --git a/Engine/Plugins/UbaController/Source/UbaController/Private/UbaHordeMetaClient.cpp b/Engine/Plugins/UbaController/Source/UbaController/Private/UbaHordeMetaClient.cpp
--- a/Engine/Plugins/UbaController/Source/UbaController/Private/UbaHordeMetaClient.cpp
+++ b/Engine/Plugins/UbaController/Source/UbaController/Private/UbaHordeMetaClient.cpp
@@ -33,21 +33,26 @@ bool FUbaHordeMetaClient::RefreshHttpClient()
 		if (FApp::IsUnattended())
 		{
 			UE_LOG(LogUbaHorde, Display, TEXT("Logging in to Horde server with environment variable UE_HORDE_TOKEN: %s"), *ServerUrl);
-			if (!HttpClient->LoginWithEnvironmentVariable())
+			if (HttpClient->LoginWithEnvironmentVariable())
 			{
-				UE_LOG(LogUbaHorde, Warning, TEXT("Login to Horde server [%s] failed"), *ServerUrl);
-				return false;
+				return true;
 			}
+
+			// Fallback to OIDC login as some commandlets might add -Unattended implicitly, but not everyone has UE_HORDE_TOKEN environment variable available
+			UE_LOG(LogUbaHorde, Display, TEXT("Login attempt with environment variable failed; trying to login with OIDC: %s"), *ServerUrl);
 		}
 		else
 		{
 			UE_LOG(LogUbaHorde, Display, TEXT("Logging in to Horde server with OIDC: %s"), *ServerUrl);
-			if (!HttpClient->LoginWithOidc(*OAuthProviderIdentifier, FApp::IsUnattended()))
-			{
-				UE_LOG(LogUbaHorde, Warning, TEXT("Login to Horde server [%s] failed"), *ServerUrl);
-				return false;
-			}
 		}
+
+		if (HttpClient->LoginWithOidc(*OAuthProviderIdentifier, FApp::IsUnattended()))
+		{
+			return true;
+		}
+
+		UE_LOG(LogUbaHorde, Warning, TEXT("Login to Horde server [%s] failed"), *ServerUrl);
+		return false;
 	}
 	return true;
 }
@@ -71,7 +76,7 @@ TSharedPtr<FUbaHordeMetaClient::HordeMachinePromise, ESPMode::ThreadSafe> FUbaHo
 	UE_LOG(LogUbaHorde, Log, TEXT("Requesting Horde agent with JSON descriptor: '%s'"), *Body);
 
 	Request->OnProcessRequestComplete().BindLambda(
-		[this, Promise](FHttpRequestPtr /*Request*/, FHttpResponsePtr HttpResponse, bool bSucceeded)
+		[this, Promise](FHttpRequestPtr /*Request*/, FHttpResponsePtr HttpResponse, bool bConnectedSuccessfully)
 		{
 			FHordeRemoteMachineInfo Info;
 			Info.Ip = TEXT("");
@@ -79,7 +84,7 @@ TSharedPtr<FUbaHordeMetaClient::HordeMachinePromise, ESPMode::ThreadSafe> FUbaHo
 			Info.bRunsWindowOS = false;
 			FMemory::Memset(Info.Nonce, 0, sizeof(Info.Nonce));
 
-			if (!bSucceeded || !HttpResponse.IsValid())
+			if (!bConnectedSuccessfully || !HttpResponse.IsValid())
 			{
 				UE_LOG(LogUbaHorde, Verbose, TEXT("No response from Horde"));
 				Promise->SetValue(MakeTuple(HttpResponse, Info));
