[UBA] Horde bundle locator must be interpreted as fixed size string view instead of a NUL-terminated string.

#rnx

diff --git a/Engine/Plugins/UbaController/Source/UbaController/Private/UbaHordeAgent.cpp b/Engine/Plugins/UbaController/Source/UbaController/Private/UbaHordeAgent.cpp
--- a/Engine/Plugins/UbaController/Source/UbaController/Private/UbaHordeAgent.cpp
+++ b/Engine/Plugins/UbaController/Source/UbaController/Private/UbaHordeAgent.cpp
@@ -94,7 +94,7 @@ bool FUbaHordeAgent::UploadBinaries(const FString& BundleDirectory, const char*
 
 	auto FindOrAddBlobFile = [&BlobFiles, &BundleDirectory](const AgentMessage::FBlobRequest& BlobRequest) -> FArchive*
 		{
-			const FString Locator(BlobRequest.Locator.GetData());
+			const FString Locator(BlobRequest.Locator.Len(), BlobRequest.Locator.GetData());
 
 			if (Locator.IsEmpty())
 			{
@@ -108,7 +108,7 @@ bool FUbaHordeAgent::UploadBinaries(const FString& BundleDirectory, const char*
 			}
 			else
 			{
-				FString Path = FPaths::Combine(BundleDirectory, BlobRequest.Locator.GetData());
+				FString Path = FPaths::Combine(BundleDirectory, Locator);
 				Path.Append(TEXT(".blob"));
 
 				if (FArchive* Archive = IFileManager::Get().CreateFileReader(*Path))
