[UbaController]
* Changed how local number of local processes is calculated based on how many remotes. Previous equation was way to aggressive towards turning off local processes

diff --git a/Engine/Plugins/UbaController/Source/UbaController/Private/UbaJobProcessor.cpp b/Engine/Plugins/UbaController/Source/UbaController/Private/UbaJobProcessor.cpp
--- a/Engine/Plugins/UbaController/Source/UbaController/Private/UbaJobProcessor.cpp
+++ b/Engine/Plugins/UbaController/Source/UbaController/Private/UbaJobProcessor.cpp
@@ -398,7 +398,10 @@ uint32 FUbaJobProcessor::Run()
 
 		if (bShouldProcessJobs)
 		{
-			int32 MaxLocal = FMath::Max(0, int32(MaxLocalParallelJobs / 2) - int32(activeRemote / 10));
+			int32 MaxPossible = FPlatformMisc::NumberOfCoresIncludingHyperthreads();
+			int32 LocalCoresToNotUse = 1 + activeRemote / 30; // Use one core per 30 remote ones
+			int32 MaxLocal = FMath::Max(0, MaxPossible - LocalCoresToNotUse);
+			MaxLocal = FMath::Min(MaxLocal, MaxLocalParallelJobs);
 			Scheduler_SetMaxLocalProcessors(UbaScheduler, MaxLocal);
 
 			int32 TargetCoreCount = FMath::Max(0, int32(queued + active) - MaxLocal);
