[UBA] Simplify generation of UBA trace files for shader compilation:
- Dump all UBA trace files into ShaderDebugInfo/UbaTrace folder.
- Rename CVar r.UbaController.TraceFilename (FString) to r.UbaController.DumpTraceFiles (bool) to simplify generation of trace files.
- Unify output filenames for trace files to always have file extension *.uba.

#rnx

diff --git a/Engine/Plugins/UbaController/Source/UbaController/Private/UbaControllerModule.cpp b/Engine/Plugins/UbaController/Source/UbaController/Private/UbaControllerModule.cpp
--- a/Engine/Plugins/UbaController/Source/UbaController/Private/UbaControllerModule.cpp
+++ b/Engine/Plugins/UbaController/Source/UbaController/Private/UbaControllerModule.cpp
@@ -18,6 +18,33 @@ DEFINE_LOG_CATEGORY(LogUbaController);
 namespace UbaControllerModule
 {
 	static constexpr int32 SubFolderCount = 32;
+
+	static bool bDumpTraceFiles = false;
+	static FAutoConsoleVariableRef CVarDumpTraceFiles(
+		TEXT("r.UbaController.DumpTraceFiles"),
+		bDumpTraceFiles,
+		TEXT("If true, UBA controller dumps trace files for later use with UBA visualizer in the Saved folder under UbaController"));
+
+	static FString MakeAndGetDebugInfoPath()
+	{
+		// Build machines should dump to the AutomationTool/Saved/Logs directory and they will upload as build artifacts via the AutomationTool.
+		FString BaseDebugInfoPath = FPaths::ProjectSavedDir();
+		if (GIsBuildMachine)
+		{
+			BaseDebugInfoPath = FPaths::Combine(*FPaths::EngineDir(), TEXT("Programs"), TEXT("AutomationTool"), TEXT("Saved"), TEXT("Logs"));
+		}
+
+		FString AbsoluteDebugInfoDirectory = IFileManager::Get().ConvertToAbsolutePathForExternalAppForWrite(*(BaseDebugInfoPath / TEXT("UbaController")));
+		FPaths::NormalizeDirectoryName(AbsoluteDebugInfoDirectory);
+
+		// Create directory if it doesn't exit yet
+		if (!IFileManager::Get().DirectoryExists(*AbsoluteDebugInfoDirectory))
+		{
+			IFileManager::Get().MakeDirectory(*AbsoluteDebugInfoDirectory, true);
+		}
+
+		return AbsoluteDebugInfoDirectory;
+	}
 };
 
 FUbaControllerModule::FUbaControllerModule()
@@ -177,6 +204,10 @@ void FUbaControllerModule::InitializeController()
 				IFileManager::Get().MakeDirectory(*FString::Printf(TEXT("%s/%d"), *WorkingDirectory, It));
 			}
 
+			if (UbaControllerModule::bDumpTraceFiles)
+			{
+				DebugInfoPath = UbaControllerModule::MakeAndGetDebugInfoPath();
+			}
 
 			JobDispatcherThread = MakeShared<FUbaJobProcessor>(*this);
 			JobDispatcherThread->StartThread();
diff --git a/Engine/Plugins/UbaController/Source/UbaController/Private/UbaJobProcessor.cpp b/Engine/Plugins/UbaController/Source/UbaController/Private/UbaJobProcessor.cpp
--- a/Engine/Plugins/UbaController/Source/UbaController/Private/UbaJobProcessor.cpp
+++ b/Engine/Plugins/UbaController/Source/UbaController/Private/UbaJobProcessor.cpp
@@ -36,12 +36,6 @@ namespace UbaJobProcessorOptions
 		bAutoLaunchVisualizer,
 		TEXT("If true, UBA visualizer will be launched automatically\n"));
 
-	static FString TraceFilename;
-	static FAutoConsoleVariableRef CVarTraceFilename(
-		TEXT("r.UbaController.TraceFilename"),
-		TraceFilename,
-		TEXT("The name of the trace file that uba outputs after a session"));
-
 	static bool bAllowProcessReuse = true;
 	static FAutoConsoleVariableRef CVarAllowProcessReuse(
 		TEXT("r.UbaController.AllowProcessReuse"),
@@ -296,9 +290,12 @@ void FUbaJobProcessor::StartUba()
 	info.remoteLogEnabled = UbaJobProcessorOptions::bProcessLogEnabled;
 
 	info.traceEnabled = true;
-	FString TraceOutputFile = UbaJobProcessorOptions::TraceFilename;
-	if (!TraceOutputFile.IsEmpty() && UE::GetMultiprocessId())
-		TraceOutputFile = FString::Printf(TEXT("%s_%u"), *TraceOutputFile, UE::GetMultiprocessId());
+	FString TraceOutputFile;
+	if (!ControllerModule.GetDebugInfoPath().IsEmpty())
+	{
+		static uint32 UbaLaunchCounter;
+		TraceOutputFile = ControllerModule.GetDebugInfoPath() / FString::Printf(TEXT("UbaController.run-%u.pid-%u.uba"), ++UbaLaunchCounter, UE::GetMultiprocessId());
+	}
 	info.traceOutputFile = *TraceOutputFile;
 	info.detailedTrace = UbaJobProcessorOptions::bDetailedTrace;
 	FString TraceName = FString::Printf(TEXT("UbaController_%s"), *FGuid::NewGuid().ToString(EGuidFormats::Digits));
diff --git a/Engine/Plugins/UbaController/Source/UbaController/Public/UbaControllerModule.h b/Engine/Plugins/UbaController/Source/UbaController/Public/UbaControllerModule.h
--- a/Engine/Plugins/UbaController/Source/UbaController/Public/UbaControllerModule.h
+++ b/Engine/Plugins/UbaController/Source/UbaController/Public/UbaControllerModule.h
@@ -33,6 +33,7 @@ public:
 
 	const FString& GetRootWorkingDirectory() const { return RootWorkingDirectory; }
 	const FString& GetWorkingDirectory() const { return WorkingDirectory; }
+	const FString& GetDebugInfoPath() const { return DebugInfoPath; }
 
 	bool HasTasksDispatchedOrPending() const;
 
@@ -48,6 +49,7 @@ private:
 	
 	FString RootWorkingDirectory;
 	FString WorkingDirectory;
+	FString DebugInfoPath;
 
 	TAtomic<int32> NextFileID;
 	TAtomic<int32> NextTaskID;
