[UBA] Rename UBA trace files to be sorted by cooker instance.

- This renames the trace fiels from "UbaController.run-1.pid-0.uba" to "UbaController.CookWorker-0.Session-0.uba", sorting all session trace files grouped by their cooker instance in a file manager.
- Local cooks will often only have a single trace file, but the build farm launches multiple cooker processes and dependeing on workload, they will launch multiple UBA sessions per cooker.
- Start session ID in filename at zero to be consistent with cooker ID.

#rnx

diff --git a/Engine/Plugins/UbaController/Source/UbaController/Private/UbaJobProcessor.cpp b/Engine/Plugins/UbaController/Source/UbaController/Private/UbaJobProcessor.cpp
--- a/Engine/Plugins/UbaController/Source/UbaController/Private/UbaJobProcessor.cpp
+++ b/Engine/Plugins/UbaController/Source/UbaController/Private/UbaJobProcessor.cpp
@@ -293,8 +293,8 @@ void FUbaJobProcessor::StartUba()
 	FString TraceOutputFile;
 	if (!ControllerModule.GetDebugInfoPath().IsEmpty())
 	{
-		static uint32 UbaLaunchCounter;
-		TraceOutputFile = ControllerModule.GetDebugInfoPath() / FString::Printf(TEXT("UbaController.run-%u.pid-%u.uba"), ++UbaLaunchCounter, UE::GetMultiprocessId());
+		static uint32 UbaSessionCounter;
+		TraceOutputFile = ControllerModule.GetDebugInfoPath() / FString::Printf(TEXT("UbaController.CookWorker-%u.Session-%u.uba"), UE::GetMultiprocessId(), UbaSessionCounter++);
 	}
 	info.traceOutputFile = *TraceOutputFile;
 	info.detailedTrace = UbaJobProcessorOptions::bDetailedTrace;
