[UBA] Allow higher verbosity of UBA logs to be forwarded to UE_LOG.

This changes CVar "r.UbaController.ShowUbaLog" (boolean) to "r.UbaController.LogVerbosity" (int32) with three levels of verbosity:
- 0 (Default): Always forward errors and warnings as before.
- 1: Also forward UBA session information.
- 2: Forward all UBA logs.

UBA logs may show up as follows, forwarding the log from the shader backends:
  LogUbaController: Display: LogD3D11ShaderCompiler: Warning: RemoveUnusedInputs: Failed to compile!
  LogUbaController: Display: LogD3D11ShaderCompiler: Warning: __UE_FILENAME_SENTINEL__(4171): (3) Unable to match rule!

#rnx

diff --git a/Engine/Plugins/UbaController/Source/UbaController/Private/UbaJobProcessor.cpp b/Engine/Plugins/UbaController/Source/UbaController/Private/UbaJobProcessor.cpp
--- a/Engine/Plugins/UbaController/Source/UbaController/Private/UbaJobProcessor.cpp
+++ b/Engine/Plugins/UbaController/Source/UbaController/Private/UbaJobProcessor.cpp
@@ -48,11 +48,21 @@ namespace UbaJobProcessorOptions
 		bDetailedTrace,
 		TEXT("If true, a UBA will output detailed trace\n"));
 
-	static bool bShowUbaLog = false;
+	enum EUbaLogVerbosity
+	{
+		UbaLogVerbosity_Default = 0, // foward erros and warnings only
+		UbaLogVerbosity_High, // also forward infos
+		UbaLogVerbosity_Max // forward all UBA logs to UE_LOG
+	};
+
+	static int32 UbaLogVerbosity = UbaLogVerbosity_Default;
 	static FAutoConsoleVariableRef CVarShowUbaLog(
-		TEXT("r.UbaController.ShowUbaLog"),
-		bShowUbaLog,
-		TEXT("If true, UBA log entries will be visible in the log\n"));
+		TEXT("r.UbaController.LogVerbosity"),
+		UbaLogVerbosity,
+		TEXT("Specifies how much of UBA logs is forwarded to UE logs..\n")
+		TEXT("0 - Default, only forward errrors and warnings.\n")
+		TEXT("1 - Also forward regular information about UBA sessions.\n")
+		TEXT("2 - Forward all UBA logs."));
 
 	static bool bProcessLogEnabled = false;
 	static FAutoConsoleVariableRef CVarProcessLogEnabled(
@@ -117,8 +127,14 @@ FUbaJobProcessor::FUbaJobProcessor(
 			case uba::LogEntryType_Warning:
 				UE_LOG(LogUbaController, Warning, TEXT("%s"), str);
 				break;
+			case uba::LogEntryType_Info:
+				if (UbaJobProcessorOptions::UbaLogVerbosity >= UbaJobProcessorOptions::UbaLogVerbosity_High)
+				{
+					UE_LOG(LogUbaController, Display, TEXT("%s"), str);
+				}
+				break;
 			default:
-				if (UbaJobProcessorOptions::bShowUbaLog)
+				if (UbaJobProcessorOptions::UbaLogVerbosity >= UbaJobProcessorOptions::UbaLogVerbosity_Max)
 				{
 					UE_LOG(LogUbaController, Display, TEXT("%s"), str);
 				}
@@ -333,8 +349,6 @@ void FUbaJobProcessor::ShutDownUba()
 		return;
 	}
 
-	SessionServer_PrintSummary(UbaSessionServer);
-
 	Server_Stop(UbaServer);
 
 	Scheduler_Destroy(UbaScheduler);
@@ -471,7 +485,7 @@ bool FUbaJobProcessor::ProcessOutputFile(FTask* CompileTask)
 	else
 	{
 		const FString OutputFileName = CompileTask != nullptr ? CompileTask->CommandData.OutputFileName : TEXT("Invalid CompileTask, cannot retrieve name");
-		UE_LOG(LogUbaController, Error, TEXT("Output File [%s] is invalid or do not exists"), *OutputFileName);
+		UE_LOG(LogUbaController, Error, TEXT("Output File [%s] is invalid or does not exist"), *OutputFileName);
 		return false;
 	}
 
