[UBA] Make UbaController plugin engine wide available.

- Set "EnabledByDefault" to true in UbaController plugin to make it engine-wide available and make per-project setting obsolete.
- Move UbaController sources from UbaController/Source/UbaController/ into UbaController/Source/ to match folder pattern of XGEController plugin.
- Remove explicit plugin setting from previous project files as it is now obsolete.

#rnx

diff --git a/Engine/Plugins/UbaController/Source/Private/UbaControllerModule.cpp b/Engine/Plugins/UbaController/Source/Private/UbaControllerModule.cpp
new file mode 100644
--- /dev/null
+++ b/Engine/Plugins/UbaController/Source/Private/UbaControllerModule.cpp
@@ -0,0 +1,268 @@
+// Copyright Epic Games, Inc. All Rights Reserved.
+
+#include "UbaControllerModule.h"
+
+#include "Misc/ConfigCacheIni.h"
+#include "UbaJobProcessor.h"
+#include "Features/IModularFeatures.h"
+#include "HAL/FileManager.h"
+#include "Misc/CommandLine.h"
+#include "Misc/CoreDelegates.h"
+#include "Misc/Paths.h"
+#include "Modules/ModuleManager.h"
+#include "CoreMinimal.h"
+
+
+DEFINE_LOG_CATEGORY(LogUbaController);
+
+namespace UbaControllerModule
+{
+	static constexpr int32 SubFolderCount = 32;
+
+	static bool bDumpTraceFiles = false;
+	static FAutoConsoleVariableRef CVarDumpTraceFiles(
+		TEXT("r.UbaController.DumpTraceFiles"),
+		bDumpTraceFiles,
+		TEXT("If true, UBA controller dumps trace files for later use with UBA visualizer in the Saved folder under UbaController"));
+
+	static FString MakeAndGetDebugInfoPath()
+	{
+		// Build machines should dump to the AutomationTool/Saved/Logs directory and they will upload as build artifacts via the AutomationTool.
+		FString BaseDebugInfoPath = FPaths::ProjectSavedDir();
+		if (GIsBuildMachine)
+		{
+			BaseDebugInfoPath = FPaths::Combine(*FPaths::EngineDir(), TEXT("Programs"), TEXT("AutomationTool"), TEXT("Saved"), TEXT("Logs"));
+		}
+
+		FString AbsoluteDebugInfoDirectory = IFileManager::Get().ConvertToAbsolutePathForExternalAppForWrite(*(BaseDebugInfoPath / TEXT("UbaController")));
+		FPaths::NormalizeDirectoryName(AbsoluteDebugInfoDirectory);
+
+		// Create directory if it doesn't exit yet
+		if (!IFileManager::Get().DirectoryExists(*AbsoluteDebugInfoDirectory))
+		{
+			IFileManager::Get().MakeDirectory(*AbsoluteDebugInfoDirectory, true);
+		}
+
+		return AbsoluteDebugInfoDirectory;
+	}
+};
+
+FUbaControllerModule::FUbaControllerModule()
+	: bSupported(false)
+	, bModuleInitialized(false)
+	, bControllerInitialized(false)
+	, RootWorkingDirectory(FPaths::Combine(FPlatformProcess::UserTempDir(), TEXT("UbaControllerWorkingDir")))
+	, WorkingDirectory(FPaths::Combine(RootWorkingDirectory, FGuid::NewGuid().ToString(EGuidFormats::Digits)))
+	, NextFileID(0)
+	, NextTaskID(0)
+{
+}
+
+FUbaControllerModule::~FUbaControllerModule()
+{	
+	if (JobDispatcherThread)
+	{
+		JobDispatcherThread->Stop();
+		// Wait until the thread is done
+		FPlatformProcess::ConditionalSleep([&](){ return JobDispatcherThread && JobDispatcherThread->IsWorkDone(); },0.1f);
+	}
+
+	CleanWorkingDirectory();
+}
+
+bool FUbaControllerModule::IsSupported()
+{
+	if (bControllerInitialized)
+	{
+		return bSupported;
+	}
+	
+#if PLATFORM_WINDOWS
+	bool bEnabled = false;
+	if (!FParse::Param(FCommandLine::Get(), TEXT("NoUbaController")))
+	{
+		GConfig->GetBool(TEXT("UbaController"), TEXT("Enabled"), bEnabled, GEngineIni);
+	}
+
+	bSupported = FPlatformProcess::SupportsMultithreading() && bEnabled;
+#else
+	// Not supported on other platforms.
+	bSupported = false;
+#endif
+	return bSupported;
+}
+
+void FUbaControllerModule::CleanWorkingDirectory() const
+{
+	if (UE::GetMultiprocessId() != 0) // Only director is allowed to clean
+	{
+		return;
+	}
+
+	IFileManager& FileManager = IFileManager::Get();
+	
+	if (!RootWorkingDirectory.IsEmpty())
+	{
+		if (!FileManager.DeleteDirectory(*RootWorkingDirectory, false, true))
+		{
+			UE_LOG(LogUbaController, Log, TEXT("%s => Failed to delete current working Directory => %s"), ANSI_TO_TCHAR(__FUNCTION__), *RootWorkingDirectory);
+		}	
+	}
+}
+
+bool FUbaControllerModule::HasTasksDispatchedOrPending() const
+{
+	return !PendingRequestedCompilationTasks.IsEmpty() || (JobDispatcherThread.IsValid() && JobDispatcherThread->HasJobsInFlight());
+}
+
+FString GetUbaBinariesPath()
+{
+#if PLATFORM_CPU_ARM_FAMILY
+	const TCHAR* BinariesArch = TEXT("arm64");
+#else
+	const TCHAR* BinariesArch = TEXT("x64");
+#endif
+	return FPaths::Combine(FPaths::EngineDir(), "Binaries/Win64/UnrealBuildAccelerator", BinariesArch);
+}
+
+void FUbaControllerModule::LoadDependencies()
+{
+	const FString UbaBinariesPath = GetUbaBinariesPath();
+	FPlatformProcess::AddDllDirectory(*UbaBinariesPath);
+	FPlatformProcess::GetDllHandle(*(FPaths::Combine(UbaBinariesPath, "UbaHost.dll")));
+}
+
+void FUbaControllerModule::StartupModule()
+{
+	check(!bModuleInitialized);
+
+	LoadDependencies();
+
+	IModularFeatures::Get().RegisterModularFeature(GetModularFeatureType(), this);
+
+	bModuleInitialized = true;
+
+	FCoreDelegates::OnEnginePreExit.AddLambda([&]()
+	{
+		if (bControllerInitialized && JobDispatcherThread)
+		{
+			JobDispatcherThread->Stop();
+			FPlatformProcess::ConditionalSleep([&]() { return JobDispatcherThread && JobDispatcherThread->IsWorkDone(); }, 0.1f);
+			JobDispatcherThread = nullptr;
+		}
+	});
+}
+
+void FUbaControllerModule::ShutdownModule()
+{
+	check(bModuleInitialized);
+	
+	IModularFeatures::Get().UnregisterModularFeature(GetModularFeatureType(), this);
+	
+	if (bControllerInitialized)
+	{
+		// Stop the jobs thread
+		if (JobDispatcherThread)
+		{
+			JobDispatcherThread->Stop();	
+			// Wait until the thread is done
+			FPlatformProcess::ConditionalSleep([&](){ return JobDispatcherThread && JobDispatcherThread->IsWorkDone(); },0.1f);
+		}
+
+		FTask* Task;
+		while (PendingRequestedCompilationTasks.Dequeue(Task))
+		{
+			FDistributedBuildTaskResult Result;
+			Result.ReturnCode = 0;
+			Result.bCompleted = false;
+			Task->Promise.SetValue(Result);
+			delete Task;
+		}
+
+		PendingRequestedCompilationTasks.Empty();
+	}
+
+	CleanWorkingDirectory();
+	bModuleInitialized = false;
+	bControllerInitialized = false;
+}
+
+void FUbaControllerModule::InitializeController()
+{
+	// We should never Initialize the controller twice
+	if (ensureAlwaysMsgf(!bControllerInitialized, TEXT("Multiple initialization of UBA controller!")))
+	{
+		CleanWorkingDirectory();
+
+		if (IsSupported())
+		{
+			IFileManager::Get().MakeDirectory(*WorkingDirectory, true);
+
+			// Pre-create the directories so we don't have to explicitly register them to uba later
+			for (int32 It = 0; It != UbaControllerModule::SubFolderCount; ++It)
+			{
+				IFileManager::Get().MakeDirectory(*FString::Printf(TEXT("%s/%d"), *WorkingDirectory, It));
+			}
+
+			if (UbaControllerModule::bDumpTraceFiles)
+			{
+				DebugInfoPath = UbaControllerModule::MakeAndGetDebugInfoPath();
+			}
+
+			JobDispatcherThread = MakeShared<FUbaJobProcessor>(*this);
+			JobDispatcherThread->StartThread();
+		}
+
+		bControllerInitialized = true;
+	}
+}
+
+FString FUbaControllerModule::CreateUniqueFilePath()
+{
+	check(bSupported);
+	int32 ID = NextFileID++;
+	int32 FolderID = ID % UbaControllerModule::SubFolderCount; // We use sub folders to be nicer to file system (we can end up with 20000 files in one folder otherwise)
+	return FString::Printf(TEXT("%s/%d/%d.uba"), *WorkingDirectory, FolderID, ID);
+}
+
+TFuture<FDistributedBuildTaskResult> FUbaControllerModule::EnqueueTask(const FTaskCommandData& CommandData)
+{
+	check(bSupported);
+
+	TPromise<FDistributedBuildTaskResult> Promise;
+	TFuture<FDistributedBuildTaskResult> Future = Promise.GetFuture();
+
+	// Enqueue the new task
+	FTask* Task = new FTask(NextTaskID++, CommandData, MoveTemp(Promise));
+	{
+		PendingRequestedCompilationTasks.Enqueue(Task);
+	}
+
+	JobDispatcherThread->HandleTaskQueueUpdated(CommandData.InputFileName);
+
+	return MoveTemp(Future);
+}
+
+bool FUbaControllerModule::PollStats(FDistributedBuildStats& OutStats)
+{
+	return JobDispatcherThread != nullptr && JobDispatcherThread->PollStats(OutStats);
+}
+
+void FUbaControllerModule::ReportJobProcessed(const FTaskResponse& InTaskResponse, FTask* CompileTask)
+{
+	if (CompileTask)
+	{
+		FDistributedBuildTaskResult Result;
+		Result.ReturnCode = InTaskResponse.ReturnCode;
+		Result.bCompleted = true;
+		CompileTask->Promise.SetValue(Result);
+		delete CompileTask;
+	}
+}
+
+UBACONTROLLER_API FUbaControllerModule& FUbaControllerModule::Get()
+{
+	return FModuleManager::LoadModuleChecked<FUbaControllerModule>(TEXT("UbaController"));
+}
+
+IMPLEMENT_MODULE(FUbaControllerModule, UbaController);
diff --git a/Engine/Plugins/UbaController/Source/Private/UbaJobProcessor.cpp b/Engine/Plugins/UbaController/Source/Private/UbaJobProcessor.cpp
new file mode 100644
--- /dev/null
+++ b/Engine/Plugins/UbaController/Source/Private/UbaJobProcessor.cpp
@@ -0,0 +1,606 @@
+// Copyright Epic Games, Inc. All Rights Reserved.
+
+#include "UbaJobProcessor.h"
+
+#include "Interfaces/ITargetPlatform.h"
+#include "Interfaces/ITargetPlatformManagerModule.h"
+#include "Misc/ConfigCacheIni.h"
+#include "Misc/CoreMisc.h"
+#include "UbaHordeAgentManager.h"
+#include "Windows/AllowWindowsPlatformTypes.h"
+#include "UbaControllerModule.h"
+#include "Windows/HideWindowsPlatformTypes.h"
+
+#include "HAL/FileManager.h"
+#include "HAL/PlatformFileManager.h"
+
+#include "Misc/Paths.h"
+#include "Misc/ScopeRWLock.h"
+
+namespace UbaJobProcessorOptions
+{
+	static float SleepTimeBetweenActions = 0.01f;
+	static FAutoConsoleVariableRef CVarSleepTimeBetweenActions(
+        TEXT("r.UbaController.SleepTimeBetweenActions"),
+        SleepTimeBetweenActions,
+        TEXT("How much time the job processor thread should sleep between actions .\n"));
+
+	static float MaxTimeWithoutTasks = 100.0f;
+	static FAutoConsoleVariableRef CVarMaxTimeWithoutTasks(
+        TEXT("r.UbaController.MaxTimeWithoutTasks"),
+        MaxTimeWithoutTasks,
+        TEXT("Time to wait (in seconds) before stop processing attempts if we don't have any pending task.\n"));
+
+	static bool bAutoLaunchVisualizer = false;
+	static FAutoConsoleVariableRef CVarAutoLaunchVisualizer(
+		TEXT("r.UbaController.AutoLaunchVisualizer"),
+		bAutoLaunchVisualizer,
+		TEXT("If true, UBA visualizer will be launched automatically\n"));
+
+	static bool bAllowProcessReuse = true;
+	static FAutoConsoleVariableRef CVarAllowProcessReuse(
+		TEXT("r.UbaController.AllowProcessReuse"),
+		bAllowProcessReuse,
+		TEXT("If true, remote process is allowed to fetch new processes from the queue (this requires the remote processes to have UbaRequestNextProcess implemented)\n"));
+
+	static bool bDetailedTrace = false;
+	static FAutoConsoleVariableRef CVarDetailedTrace(
+		TEXT("r.UbaController.DetailedTrace"),
+		bDetailedTrace,
+		TEXT("If true, a UBA will output detailed trace\n"));
+
+	enum EUbaLogVerbosity
+	{
+		UbaLogVerbosity_Default = 0, // foward erros and warnings only
+		UbaLogVerbosity_High, // also forward infos
+		UbaLogVerbosity_Max // forward all UBA logs to UE_LOG
+	};
+
+	static int32 UbaLogVerbosity = UbaLogVerbosity_Default;
+	static FAutoConsoleVariableRef CVarShowUbaLog(
+		TEXT("r.UbaController.LogVerbosity"),
+		UbaLogVerbosity,
+		TEXT("Specifies how much of UBA logs is forwarded to UE logs..\n")
+		TEXT("0 - Default, only forward errrors and warnings.\n")
+		TEXT("1 - Also forward regular information about UBA sessions.\n")
+		TEXT("2 - Forward all UBA logs."));
+
+	static bool bProcessLogEnabled = false;
+	static FAutoConsoleVariableRef CVarProcessLogEnabled(
+		TEXT("r.UbaController.ProcessLogEnabled"),
+		bProcessLogEnabled,
+		TEXT("If true, each detoured process will write a log file. Note this is only useful if UBA is compiled in debug\n"));
+
+	FString ReplaceEnvironmentVariablesInPath(const FString& ExtraFilePartialPath) // Duplicated code with FAST build.. put it somewhere else?
+	{
+		FString ParsedPath;
+
+		// Fast build cannot read environmental variables easily
+		// Is better to resolve them here
+		if (ExtraFilePartialPath.Contains(TEXT("%")))
+		{
+			TArray<FString> PathSections;
+			ExtraFilePartialPath.ParseIntoArray(PathSections, TEXT("/"));
+
+			for (FString& Section : PathSections)
+			{
+				if (Section.Contains(TEXT("%")))
+				{
+					Section.RemoveFromStart(TEXT("%"));
+					Section.RemoveFromEnd(TEXT("%"));
+					Section = FPlatformMisc::GetEnvironmentVariable(*Section);
+				}
+			}
+
+			for (FString& Section : PathSections)
+			{
+				ParsedPath /= Section;
+			}
+
+			FPaths::NormalizeDirectoryName(ParsedPath);
+		}
+
+		if (ParsedPath.IsEmpty())
+		{
+			ParsedPath = ExtraFilePartialPath;
+		}
+
+		return ParsedPath;
+	}
+}
+
+FUbaJobProcessor::FUbaJobProcessor(
+	FUbaControllerModule& InControllerModule) :
+
+	Thread(nullptr),
+	ControllerModule(InControllerModule),
+	bForceStop(false),
+	LastTimeCheckedForTasks(0),
+	bShouldProcessJobs(false),
+	bIsWorkDone(false),
+	LogWriter([]() {}, []() {}, [](uba::LogEntryType type, const wchar_t* str, uba::u32 strlen)
+	{
+			switch (type)
+			{
+			case uba::LogEntryType_Error:
+				UE_LOG(LogUbaController, Error, TEXT("%s"), str);
+				break;
+			case uba::LogEntryType_Warning:
+				UE_LOG(LogUbaController, Warning, TEXT("%s"), str);
+				break;
+			case uba::LogEntryType_Info:
+				if (UbaJobProcessorOptions::UbaLogVerbosity >= UbaJobProcessorOptions::UbaLogVerbosity_High)
+				{
+					UE_LOG(LogUbaController, Display, TEXT("%s"), str);
+				}
+				break;
+			default:
+				if (UbaJobProcessorOptions::UbaLogVerbosity >= UbaJobProcessorOptions::UbaLogVerbosity_Max)
+				{
+					UE_LOG(LogUbaController, Display, TEXT("%s"), str);
+				}
+				break;
+			}
+	})
+{
+	Uba_SetCustomAssertHandler([](const uba::tchar* text)
+		{
+			checkf(false, TEXT("%s"), text);
+		});
+
+	if (!GConfig->GetInt(TEXT("UbaController"), TEXT("MaxLocalParallelJobs"), MaxLocalParallelJobs, GEngineIni))
+	{
+		MaxLocalParallelJobs = -1;
+	}
+
+	if (MaxLocalParallelJobs == -1)
+	{
+		MaxLocalParallelJobs = FPlatformMisc::NumberOfCoresIncludingHyperthreads();
+	}
+}
+
+FUbaJobProcessor::~FUbaJobProcessor()
+{
+	delete Thread;
+}
+
+void FUbaJobProcessor::CalculateKnownInputs()
+{
+	// TODO: This is ShaderCompileWorker specific and this code is designed to handle all kinds of distributed workload.
+	// Instead this information should be provided from the outside
+
+
+	if (KnownInputsCount) // In order to improve startup we provide some of the input we know will be loaded by ShaderCompileWorker.exe
+	{
+		return;
+	}
+
+	auto AddKnownInput = [&](const FString& file)
+		{
+			auto& fileData = file.GetCharArray();
+			auto num = KnownInputsBuffer.Num();
+			KnownInputsBuffer.SetNum(num + fileData.Num());
+			memcpy(KnownInputsBuffer.GetData() + num, fileData.GetData(), fileData.Num() * sizeof(uba::tchar));
+			++KnownInputsCount;
+		};
+
+	// Get the binaries
+	TArray<FString> KnownFileNames;
+	FString BinDir = FPaths::Combine(FPaths::EngineDir(), TEXT("Binaries/Win64")); // TODO: Need to support other targets than Win64
+	IFileManager::Get().FindFilesRecursive(KnownFileNames, *BinDir, TEXT("ShaderCompileWorker*.*"), true, false);
+	for (const FString& file : KnownFileNames)
+	{
+		if (!file.EndsWith(TEXT(".pdb")))
+		{
+			AddKnownInput(file);
+		}
+	}
+
+	// Get the compiler dependencies for all platforms)
+	ITargetPlatformManagerModule* TargetPlatformManager = GetTargetPlatformManager();
+	for (ITargetPlatform* TargetPlatform : GetTargetPlatformManager()->GetTargetPlatforms())
+	{
+		KnownFileNames.Empty();
+		TargetPlatform->GetShaderCompilerDependencies(KnownFileNames);
+
+		for (const FString& ExtraFilePartialPath : KnownFileNames)
+		{
+			if (!ExtraFilePartialPath.Contains(TEXT("*"))) // Seems like there are some *.x paths in there.. TODO: Do a find files
+			{
+				AddKnownInput(UbaJobProcessorOptions::ReplaceEnvironmentVariablesInPath(ExtraFilePartialPath));
+			}
+		}
+	}
+
+	// Get all the config files
+	for (const FString& ConfigDir : FPaths::GetExtensionDirs(FPaths::EngineDir(), TEXT("Config")))
+	{
+		KnownFileNames.Empty();
+		IFileManager::Get().FindFilesRecursive(KnownFileNames, *ConfigDir, TEXT("*.ini"), true, false);
+		for (const FString& file : KnownFileNames)
+		{
+			AddKnownInput(file);
+		}
+	}
+
+	KnownInputsBuffer.Add(0);
+}
+
+void FUbaJobProcessor::RunTaskWithUba(FTask* Task)
+{
+	FTaskCommandData& Data = Task->CommandData;
+	SessionServer_RegisterNewFile(UbaSessionServer, *Data.InputFileName);
+
+	FString InputFileName = FPaths::GetCleanFilename(Data.InputFileName);
+	FString OutputFileName = FPaths::GetCleanFilename(Data.OutputFileName);
+	FString Parameters = FString::Printf(TEXT("\"%s/\" %d 0 \"%s\" \"%s\" %s "), *Data.WorkingDirectory, Data.DispatcherPID, *InputFileName, *OutputFileName, *Data.ExtraCommandArgs);
+	FString AppDir = FPaths::GetPath(Data.Command);
+
+	uba::ProcessStartInfo ProcessInfo;
+	ProcessInfo.application = *Data.Command;
+	ProcessInfo.arguments = *Parameters;
+	ProcessInfo.description = *InputFileName;
+	ProcessInfo.workingDir = *AppDir;
+	ProcessInfo.writeOutputFilesOnFail = true;
+
+	if (UbaJobProcessorOptions::bProcessLogEnabled)
+	{
+		ProcessInfo.logFile = *InputFileName;
+	}
+	
+	struct ExitedInfo
+	{
+		FUbaJobProcessor* Processor;
+		FString InputFile;
+		FString OutputFile;
+		FTask* Task;
+	};
+
+	auto Info = new ExitedInfo;
+	Info->Processor = this;
+	Info->InputFile = Data.InputFileName;
+	Info->OutputFile = Data.OutputFileName;
+	Info->Task = Task;
+
+	ProcessInfo.userData = Info;
+	ProcessInfo.exitedFunc = [](void* userData, const uba::ProcessHandle& ph)
+		{
+			uint32 logLineIndex = 0;
+			while (const uba::tchar* line = ProcessHandle_GetLogLine(&ph, logLineIndex++))
+			{
+				UE_LOG(LogUbaController, Display, TEXT("%s"), line);
+			}
+
+			if (auto Info = (ExitedInfo*)userData) // It can be null if custom message has already handled all of them
+			{
+				IFileManager::Get().Delete(*Info->InputFile);
+				SessionServer_RegisterDeleteFile(Info->Processor->UbaSessionServer, *Info->InputFile);
+				Info->Processor->HandleUbaJobFinished(Info->Task);
+
+				Storage_DeleteFile(Info->Processor->UbaStorageServer, *Info->InputFile);
+				Storage_DeleteFile(Info->Processor->UbaStorageServer, *Info->OutputFile);
+
+				delete Info;
+			}
+		};
+
+
+	Scheduler_EnqueueProcess(UbaScheduler, ProcessInfo, 1.0f, KnownInputsBuffer.GetData(), KnownInputsBuffer.Num()*sizeof(uba::tchar), KnownInputsCount);
+}
+
+FString GetUbaBinariesPath();
+
+void FUbaJobProcessor::StartUba()
+{
+	checkf(UbaServer == nullptr, TEXT("FUbaJobProcessor::StartUba() was called twice before FUbaJobProcessor::ShutDownUba()"));
+
+	UbaServer = CreateServer(LogWriter);
+
+	FString RootDir = FString::Printf(TEXT("%s/%s/%u"), FPlatformProcess::UserTempDir(), TEXT("UbaControllerStorageDir"), UE::GetMultiprocessId());
+	IFileManager::Get().MakeDirectory(*RootDir, true);
+
+	uba::u64 casCapacityBytes = 32llu * 1024 * 1024 * 1024;
+	UbaStorageServer = CreateStorageServer(*UbaServer, *RootDir, casCapacityBytes, true, LogWriter);
+
+	uba::SessionServerCreateInfo info(*UbaStorageServer, *UbaServer, LogWriter);
+	info.launchVisualizer = UbaJobProcessorOptions::bAutoLaunchVisualizer;
+	info.rootDir = *RootDir;
+	info.allowMemoryMaps = false; // Skip using memory maps
+	info.remoteLogEnabled = UbaJobProcessorOptions::bProcessLogEnabled;
+
+	info.traceEnabled = true;
+	FString TraceOutputFile;
+	if (!ControllerModule.GetDebugInfoPath().IsEmpty())
+	{
+		static uint32 UbaSessionCounter;
+		TraceOutputFile = ControllerModule.GetDebugInfoPath() / FString::Printf(TEXT("UbaController.CookWorker-%u.Session-%u.uba"), UE::GetMultiprocessId(), UbaSessionCounter++);
+	}
+	info.traceOutputFile = *TraceOutputFile;
+	info.detailedTrace = UbaJobProcessorOptions::bDetailedTrace;
+	FString TraceName = FString::Printf(TEXT("UbaController_%s"), *FGuid::NewGuid().ToString(EGuidFormats::Digits));
+	info.traceName = *TraceName;
+
+
+	//info.remoteLogEnabled = true;
+	UbaSessionServer = CreateSessionServer(info);
+
+	CalculateKnownInputs();
+
+	UbaScheduler = Scheduler_Create(UbaSessionServer, MaxLocalParallelJobs, UbaJobProcessorOptions::bAllowProcessReuse);
+	Scheduler_Start(UbaScheduler);
+
+	HandleTaskQueueUpdated(TEXT("")); // Flush tasks into uba scheduler
+
+	if (UE::GetMultiprocessId() == 0)
+	{
+		Server_StartListen(UbaServer, uba::DefaultPort, nullptr); // Start listen so any helper on the LAN can join in
+	}
+
+	HordeAgentManager = MakeUnique<FUbaHordeAgentManager>(ControllerModule.GetWorkingDirectory(), GetUbaBinariesPath());
+
+	auto AddClientCallback = [](void* userData, const uba::tchar* ip, uint16 port)
+		{
+			return Server_AddClient((uba::NetworkServer*)userData, ip, port, nullptr);
+		};
+	HordeAgentManager->SetAddClientCallback(AddClientCallback, UbaServer);
+
+	FString HordeConfig;
+	if (GConfig->GetString(TEXT("UbaController"), TEXT("Horde"), HordeConfig, GEngineIni))
+	{
+		HordeConfig.TrimStartInline();
+		HordeConfig.TrimEndInline();
+		HordeConfig.RemoveFromStart(TEXT("("));
+		HordeConfig.RemoveFromEnd(TEXT(")"));
+
+		FString Url;
+		if (FParse::Value(*HordeConfig, TEXT("Url="), Url))
+		{
+			UE_LOG(LogUbaController, Log, TEXT("Found UBA controller Url: \"%s\""), *Url);
+			HordeAgentManager->SetUrl(Url);
+		}
+
+		FString Pool;
+		if (FParse::Value(*HordeConfig, TEXT("Pool="), Pool))
+		{
+			UE_LOG(LogUbaController, Log, TEXT("Found UBA controller Pool: \"%s\""), *Pool);
+			HordeAgentManager->SetPool(Pool);
+		}
+
+		FString Oidc;
+		if (FParse::Value(*HordeConfig, TEXT("Oidc="), Oidc))
+		{
+			UE_LOG(LogUbaController, Log, TEXT("Found UBA controller Oidc: \"%s\""), *Oidc);
+			HordeAgentManager->SetOidc(Oidc);
+		}
+
+		uint32 MaxCores = 0;
+		if (FParse::Value(*HordeConfig, TEXT("MaxCores="), MaxCores))
+		{
+			UE_LOG(LogUbaController, Log, TEXT("Found UBA controller MaxCores: \"%u\""), MaxCores);
+			HordeAgentManager->SetMaxCoreCount(MaxCores);
+		}
+	}
+
+
+	UE_LOG(LogUbaController, Display, TEXT("Created UBA storage server: RootDir=%s"), *RootDir);
+}
+
+void FUbaJobProcessor::ShutDownUba()
+{
+	UE_LOG(LogUbaController, Display, TEXT("Shutting down UBA/Horde connection"));
+
+	HordeAgentManager = nullptr;
+
+	if (UbaSessionServer == nullptr)
+	{
+		return;
+	}
+
+	Server_Stop(UbaServer);
+
+	Scheduler_Destroy(UbaScheduler);
+	DestroySessionServer(UbaSessionServer);
+	DestroyStorageServer(UbaStorageServer);
+	DestroyServer(UbaServer);
+
+	UbaScheduler = nullptr;
+	UbaSessionServer = nullptr;
+	UbaStorageServer = nullptr;
+	UbaServer = nullptr;
+}
+
+uint32 FUbaJobProcessor::Run()
+{
+	bIsWorkDone = false;
+	
+	uint32 LastTimeSinceHadJobs = FPlatformTime::Cycles();	
+
+	while (!bForceStop)
+	{
+		const float ElapsedSeconds = (FPlatformTime::Cycles() - LastTimeSinceHadJobs) * FPlatformTime::GetSecondsPerCycle();
+
+		uint32 queued = 0;
+		uint32 activeLocal = 0;
+		uint32 activeRemote = 0;
+		uint32 finished = 0;
+
+		FScopeLock lock(&bShouldProcessJobsLock);
+
+		if (UbaScheduler)
+		{
+			Scheduler_GetStats(UbaScheduler, queued, activeLocal, activeRemote, finished);
+		}
+		uint32 active = activeLocal + activeRemote;
+
+		// We don't want to hog up Horde resources.
+		if (bShouldProcessJobs && ElapsedSeconds > UbaJobProcessorOptions::MaxTimeWithoutTasks && (queued + active) == 0)
+		{
+			// If we're optimizing job starting, we only want to shutdown UBA if all the processes have terminated
+			bShouldProcessJobs = false;
+			ShutDownUba();
+		}
+
+		// Check if we have new tasks to process
+		if ((ControllerModule.HasTasksDispatchedOrPending() || (queued + active) != 0))
+		{
+			if (!bShouldProcessJobs)
+			{
+				// We have new tasks. Start processing again
+				StartUba();
+
+				bShouldProcessJobs = true;
+			}
+
+			LastTimeSinceHadJobs = FPlatformTime::Cycles();
+		}
+
+		if (bShouldProcessJobs)
+		{
+			int32 MaxPossible = FPlatformMisc::NumberOfCoresIncludingHyperthreads();
+			int32 LocalCoresToNotUse = 1 + activeRemote / 30; // Use one core per 30 remote ones
+			int32 MaxLocal = FMath::Max(0, MaxPossible - LocalCoresToNotUse);
+			MaxLocal = FMath::Min(MaxLocal, MaxLocalParallelJobs);
+			Scheduler_SetMaxLocalProcessors(UbaScheduler, MaxLocal);
+
+			int32 TargetCoreCount = FMath::Max(0, int32(queued + active) - MaxLocal);
+
+			HordeAgentManager->SetTargetCoreCount(TargetCoreCount);
+			
+			// TODO: Not sure this is a good idea in a cooking scenario where number of queued processes are going up and down
+			SessionServer_SetMaxRemoteProcessCount(UbaSessionServer, TargetCoreCount);
+
+			UpdateStats();
+		}
+
+		lock.Unlock();
+
+		FPlatformProcess::Sleep(UbaJobProcessorOptions::SleepTimeBetweenActions);
+	}
+
+	ShutDownUba();
+
+	bIsWorkDone = true;
+	return 0;
+}
+
+void FUbaJobProcessor::StartThread()
+{
+	Thread = FRunnableThread::Create(this, TEXT("UbaJobProcessor"), 0, TPri_SlightlyBelowNormal, FPlatformAffinity::GetPoolThreadMask());
+}
+
+bool FUbaJobProcessor::ProcessOutputFile(FTask* CompileTask)
+{
+	//TODO: This method is mostly taken from the other Distribution controllers
+	// As we get an explicit callback when the process ends, we should be able to simplify this to just check if the file exists
+	IPlatformFile& PlatformFile = FPlatformFileManager::Get().GetPlatformFile();
+	IFileManager& FileManager = IFileManager::Get();
+
+	constexpr uint64 VersionAndFileSizeSize = sizeof(uint32) + sizeof(uint64);	
+	if (ensure(CompileTask && PlatformFile.FileExists(*CompileTask->CommandData.OutputFileName) &&
+		FileManager.FileSize(*CompileTask->CommandData.OutputFileName) > VersionAndFileSizeSize))
+	{
+		const TUniquePtr<FArchive> OutputFilePtr(FileManager.CreateFileReader(*CompileTask->CommandData.OutputFileName, FILEREAD_Silent));
+		if (ensure(OutputFilePtr))
+		{
+			FArchive& OutputFile = *OutputFilePtr;
+			int32 OutputVersion;
+			OutputFile << OutputVersion; // NOTE (SB): Do not care right now about the version.
+			int64 FileSize = 0;
+			OutputFile << FileSize;
+
+			// NOTE (SB): Check if we received the full file yet.
+			if (ensure(OutputFile.TotalSize() >= FileSize))
+			{
+				FTaskResponse TaskCompleted;
+				TaskCompleted.ID = CompileTask->ID;
+				TaskCompleted.ReturnCode = 0;
+						
+				ControllerModule.ReportJobProcessed(TaskCompleted, CompileTask);
+			}
+			else
+			{
+				UE_LOG(LogUbaController, Error, TEXT("Output file size is not correct [%s] | Expected Size [%lld] : => Actual Size : [%lld]"), *CompileTask->CommandData.OutputFileName, OutputFile.TotalSize(), FileSize);
+				return false;
+			}
+		}
+		else
+		{
+			UE_LOG(LogUbaController, Error, TEXT("Failed open for read Output File [%s]"), *CompileTask->CommandData.OutputFileName);
+			return false;
+		}
+	}
+	else
+	{
+		const FString OutputFileName = CompileTask != nullptr ? CompileTask->CommandData.OutputFileName : TEXT("Invalid CompileTask, cannot retrieve name");
+		UE_LOG(LogUbaController, Error, TEXT("Output File [%s] is invalid or does not exist"), *OutputFileName);
+		return false;
+	}
+
+	return true;
+}
+
+void FUbaJobProcessor::HandleUbaJobFinished(FTask* CompileTask)
+{
+	const bool bWasSuccessful = ProcessOutputFile(CompileTask);
+	if (!bWasSuccessful)
+	{
+		// If it failed running locally, lets try Run it locally but outside Uba
+		// Signaling a jobs as complete when it wasn't done, will cause a rerun on local worker as fallback
+		// because there is not output file for this job
+
+		FTaskResponse TaskCompleted;
+		TaskCompleted.ID = CompileTask->ID;
+		TaskCompleted.ReturnCode = 0;
+		ControllerModule.ReportJobProcessed(TaskCompleted, CompileTask);
+	}
+}
+
+void FUbaJobProcessor::HandleTaskQueueUpdated(const FString& InputFileName)
+{
+	FScopeLock lock(&bShouldProcessJobsLock);
+
+	if (!UbaScheduler)
+	{
+		return;
+	}
+
+	while (true)
+	{
+		FTask* Task = nullptr;
+		if (!ControllerModule.PendingRequestedCompilationTasks.Dequeue(Task) || !Task)
+			break;
+		RunTaskWithUba(Task);
+	}
+}
+
+bool FUbaJobProcessor::HasJobsInFlight() const
+{
+	if (!UbaScheduler)
+	{
+		return false;
+	}
+	uint32 queued = 0;
+	uint32 activeLocal = 0;
+	uint32 activeRemote = 0;
+	uint32 finished = 0;
+	Scheduler_GetStats(UbaScheduler, queued, activeLocal, activeRemote, finished);	
+	return (queued + activeLocal + activeRemote) != 0;
+}
+
+bool FUbaJobProcessor::PollStats(FDistributedBuildStats& OutStats)
+{
+	// Return current stats and reset internal data
+	FScopeLock StatsLockGuard(&StatsLock);
+	OutStats = Stats;
+	Stats = FDistributedBuildStats();
+	return true;
+}
+
+void FUbaJobProcessor::UpdateStats()
+{
+	FScopeLock StatsLockGuard(&StatsLock);
+
+	// Update maximum
+	Stats.MaxRemoteAgents = FMath::Max(Stats.MaxRemoteAgents, (uint32)HordeAgentManager->GetAgentCount());
+	Stats.MaxActiveAgentCores = FMath::Max(Stats.MaxActiveAgentCores, HordeAgentManager->GetActiveCoreCount());
+}
diff --git a/Engine/Plugins/UbaController/Source/Public/UbaControllerModule.h b/Engine/Plugins/UbaController/Source/Public/UbaControllerModule.h
new file mode 100644
--- /dev/null
+++ b/Engine/Plugins/UbaController/Source/Public/UbaControllerModule.h
@@ -0,0 +1,58 @@
+// Copyright Epic Games, Inc. All Rights Reserved.
+#pragma once
+
+#include "Containers/Queue.h"
+#include "DistributedBuildControllerInterface.h"
+
+class FUbaJobProcessor;
+struct FTaskCommandData;
+
+DECLARE_LOG_CATEGORY_EXTERN(LogUbaController, Log, Log);
+
+class FUbaControllerModule : public IDistributedBuildController , public TSharedFromThis<FUbaControllerModule>
+{
+public:
+	
+	FUbaControllerModule();
+	virtual ~FUbaControllerModule() override;
+	
+	UBACONTROLLER_API static FUbaControllerModule& Get();
+	UBACONTROLLER_API virtual bool IsSupported() override final;
+
+	UBACONTROLLER_API virtual const FString GetName() override final { return FString("UBA Controller"); };
+	
+	virtual void StartupModule() override final;
+	virtual void ShutdownModule() override final;
+	virtual void InitializeController() override final;
+	virtual FString CreateUniqueFilePath() override final;
+	virtual TFuture<FDistributedBuildTaskResult> EnqueueTask(const FTaskCommandData& CommandData) override final;
+	virtual bool PollStats(FDistributedBuildStats& OutStats) override;
+
+	void ReportJobProcessed(const FTaskResponse& InTaskResponse, FTask* CompileTask);
+	void CleanWorkingDirectory() const;
+
+	const FString& GetRootWorkingDirectory() const { return RootWorkingDirectory; }
+	const FString& GetWorkingDirectory() const { return WorkingDirectory; }
+	const FString& GetDebugInfoPath() const { return DebugInfoPath; }
+
+	bool HasTasksDispatchedOrPending() const;
+
+	// Queue of tasks submitted by the engine, but not yet dispatched to the controller.
+	TQueue<FTask*, EQueueMode::Mpsc> PendingRequestedCompilationTasks;
+
+private:
+	void LoadDependencies();
+
+	bool bSupported;
+	bool bModuleInitialized;
+	bool bControllerInitialized;
+	
+	FString RootWorkingDirectory;
+	FString WorkingDirectory;
+	FString DebugInfoPath;
+
+	TAtomic<int32> NextFileID;
+	TAtomic<int32> NextTaskID;
+
+	TSharedPtr<FUbaJobProcessor> JobDispatcherThread;
+};
diff --git a/Engine/Plugins/UbaController/Source/Public/UbaJobProcessor.h b/Engine/Plugins/UbaController/Source/Public/UbaJobProcessor.h
new file mode 100644
--- /dev/null
+++ b/Engine/Plugins/UbaController/Source/Public/UbaJobProcessor.h
@@ -0,0 +1,86 @@
+// Copyright Epic Games, Inc. All Rights Reserved.
+#pragma once
+
+#include "HAL/PlatformAtomics.h"
+#include "HAL/Runnable.h"
+#include "Templates/SharedPointer.h"
+#include "Templates/UniquePtr.h"
+#include "UbaExports.h"
+#include "DistributedBuildControllerInterface.h"
+
+class FUbaControllerModule;
+class FUbaHordeAgentManager;
+struct FTask;
+
+class FUbaJobProcessor : public FRunnable, public TSharedFromThis<FUbaJobProcessor>
+{
+public:
+	FUbaJobProcessor(FUbaControllerModule& InControllerModule);
+
+	virtual ~FUbaJobProcessor() override;
+	
+	/** Main loop */
+	virtual uint32 Run() override;
+	
+	/** Aborts the main loop as soon as possible */
+	virtual void Stop() override { bForceStop = true; };
+	
+	/** Creates the threads and starts the main loop */
+	void StartThread();
+	bool ProcessOutputFile(FTask* CompileTask);
+
+	void HandleUbaJobFinished(FTask* CompileTask);
+
+	/** Used to know when this thread has finished the main loop */
+	bool IsWorkDone() const { return bIsWorkDone;};
+
+	void HandleTaskQueueUpdated(const FString& InputFileName);
+
+	bool HasJobsInFlight() const;
+
+	bool PollStats(FDistributedBuildStats& OutStats);
+
+private:
+
+	void CalculateKnownInputs();
+	void RunTaskWithUba(FTask* Task);
+
+	void StartUba();
+	void ShutDownUba();
+
+	void UpdateStats();
+
+	/** The runnable thread */
+	FRunnableThread* Thread;
+	
+	FUbaControllerModule& ControllerModule;
+
+	int32 MaxLocalParallelJobs = 0;
+	
+	/** Used to abort the current processing loop */
+	TAtomic<bool> bForceStop;
+	
+	/** Used to abort the current processing loop */
+	uint32 LastTimeCheckedForTasks;
+	
+	/** Used to stop the processing loop in the next main loop */
+	FCriticalSection bShouldProcessJobsLock;
+	bool bShouldProcessJobs;
+
+	/** Set to true when the main loop finishes*/
+	bool bIsWorkDone;
+
+	uba::NetworkServer* UbaServer = nullptr;
+	uba::Storage* UbaStorageServer = nullptr;
+	uba::SessionServer* UbaSessionServer = nullptr;
+	uba::Scheduler* UbaScheduler = nullptr;
+	TUniquePtr<FUbaHordeAgentManager> HordeAgentManager;
+
+	uint32 KnownInputsCount = 0;
+	TArray<uba::tchar> KnownInputsBuffer;
+
+	uba::CallbackLogWriter LogWriter;
+
+	FCriticalSection StatsLock;
+	FDistributedBuildStats Stats;
+};
\ No newline at end of file
diff --git a/Engine/Plugins/UbaController/Source/UbaController.build.cs b/Engine/Plugins/UbaController/Source/UbaController.build.cs
new file mode 100644
--- /dev/null
+++ b/Engine/Plugins/UbaController/Source/UbaController.build.cs
@@ -0,0 +1,40 @@
+// Copyright Epic Games, Inc. All Rights Reserved.
+
+using System.IO;
+using UnrealBuildTool;
+
+public class UbaController : ModuleRules
+{
+	public UbaController(ReadOnlyTargetRules TargetRules) : base(TargetRules)
+	{
+		PrivateDependencyModuleNames.AddRange(new string[]
+		{
+			"Core",
+			"DistributedBuildInterface",
+			"Projects",
+			"RenderCore",
+			"TargetPlatform",
+			"Horde",
+			"HTTP",
+			"Sockets",
+			"UbaCoordinatorHorde",
+			"Json"
+		});
+
+		string Arch = TargetRules.Architecture.WindowsName == "arm64ec" ? "arm64" : TargetRules.Architecture.WindowsName;
+		string UbaSourcePath = Path.Combine(EngineDirectory, "Source", "Programs", "UnrealBuildAccelerator");
+
+		PublicIncludePaths.Add(UbaSourcePath);
+		PublicIncludePaths.Add(Path.Combine(UbaSourcePath, "Core", "Public"));
+		PublicIncludePaths.Add(Path.Combine(UbaSourcePath, "Common", "Public"));
+        PublicIncludePaths.Add(Path.Combine(UbaSourcePath, "Host", "Public"));
+        PublicIncludePathModuleNames.Add("BLAKE3");
+
+		string UbaBinariesPath = Path.Combine(EngineDirectory, "Binaries", "Win64", "UnrealBuildAccelerator", Arch);
+
+		PublicAdditionalLibraries.Add(Path.Combine(UbaBinariesPath, "UbaHost.lib"));
+		PrivateRuntimeLibraryPaths.Add(Path.Combine(ModuleDirectory, "lib"));
+
+		PublicDelayLoadDLLs.Add("UbaHost.dll");
+	}
+}
\ No newline at end of file
diff --git a/Engine/Plugins/UbaController/Source/UbaController/Private/UbaControllerModule.cpp b/Engine/Plugins/UbaController/Source/UbaController/Private/UbaControllerModule.cpp
deleted file mode 100644
diff --git a/Engine/Plugins/UbaController/Source/UbaController/Private/UbaJobProcessor.cpp b/Engine/Plugins/UbaController/Source/UbaController/Private/UbaJobProcessor.cpp
deleted file mode 100644
diff --git a/Engine/Plugins/UbaController/Source/UbaController/Public/UbaControllerModule.h b/Engine/Plugins/UbaController/Source/UbaController/Public/UbaControllerModule.h
deleted file mode 100644
diff --git a/Engine/Plugins/UbaController/Source/UbaController/Public/UbaJobProcessor.h b/Engine/Plugins/UbaController/Source/UbaController/Public/UbaJobProcessor.h
deleted file mode 100644
diff --git a/Engine/Plugins/UbaController/Source/UbaController/UbaController.build.cs b/Engine/Plugins/UbaController/Source/UbaController/UbaController.build.cs
deleted file mode 100644
diff --git a/Engine/Plugins/UbaController/UbaController.uplugin b/Engine/Plugins/UbaController/UbaController.uplugin
--- a/Engine/Plugins/UbaController/UbaController.uplugin
+++ b/Engine/Plugins/UbaController/UbaController.uplugin
@@ -11,7 +11,7 @@
 	"MarketplaceURL" : "",
 	"SupportURL" : "",
 	"IsExperimentalVersion": true,
-	"EnabledByDefault" : false,
+	"EnabledByDefault" : true,
 	"CanContainContent" : false,
 	"IsBetaVersion" : false,
 	"Installed" : false,
