[UbaController]
* Updated code to use latest uba export api

diff --git a/Engine/Plugins/UbaController/Source/Private/UbaJobProcessor.cpp b/Engine/Plugins/UbaController/Source/Private/UbaJobProcessor.cpp
--- a/Engine/Plugins/UbaController/Source/Private/UbaJobProcessor.cpp
+++ b/Engine/Plugins/UbaController/Source/Private/UbaJobProcessor.cpp
@@ -7,6 +7,8 @@
 #include "Misc/ConfigCacheIni.h"
 #include "Misc/CoreMisc.h"
 #include "UbaHordeAgentManager.h"
+#include "UbaProcessStartInfo.h"
+#include "UbaSessionServerCreateInfo.h"
 #include "Windows/AllowWindowsPlatformTypes.h"
 #include "UbaControllerModule.h"
 #include "Windows/HideWindowsPlatformTypes.h"
@@ -277,8 +279,8 @@ void FUbaJobProcessor::RunTaskWithUba(FTask* Task)
 				SessionServer_RegisterDeleteFile(Info->Processor->UbaSessionServer, *Info->InputFile);
 				Info->Processor->HandleUbaJobFinished(Info->Task);
 
-				Storage_DeleteFile(Info->Processor->UbaStorageServer, *Info->InputFile);
-				Storage_DeleteFile(Info->Processor->UbaStorageServer, *Info->OutputFile);
+				StorageServer_DeleteFile(Info->Processor->UbaStorageServer, *Info->InputFile);
+				StorageServer_DeleteFile(Info->Processor->UbaStorageServer, *Info->OutputFile);
 
 				delete Info;
 			}
@@ -294,15 +296,15 @@ void FUbaJobProcessor::StartUba()
 {
 	checkf(UbaServer == nullptr, TEXT("FUbaJobProcessor::StartUba() was called twice before FUbaJobProcessor::ShutDownUba()"));
 
-	UbaServer = CreateServer(LogWriter);
+	UbaServer = NetworkServer_Create(LogWriter);
 
 	FString RootDir = FString::Printf(TEXT("%s/%s/%u"), FPlatformProcess::UserTempDir(), TEXT("UbaControllerStorageDir"), UE::GetMultiprocessId());
 	IFileManager::Get().MakeDirectory(*RootDir, true);
 
 	uba::u64 casCapacityBytes = 32llu * 1024 * 1024 * 1024;
-	UbaStorageServer = CreateStorageServer(*UbaServer, *RootDir, casCapacityBytes, true, LogWriter);
+	UbaStorageServer = StorageServer_Create(*UbaServer, *RootDir, casCapacityBytes, true, LogWriter);
 
-	uba::SessionServerCreateInfo info(*UbaStorageServer, *UbaServer, LogWriter);
+	uba::SessionServerCreateInfo info(*(uba::Storage*)UbaStorageServer, *UbaServer, LogWriter);
 	info.launchVisualizer = UbaJobProcessorOptions::bAutoLaunchVisualizer;
 	info.rootDir = *RootDir;
 	info.allowMemoryMaps = false; // Skip using memory maps
@@ -322,7 +324,7 @@ void FUbaJobProcessor::StartUba()
 
 
 	//info.remoteLogEnabled = true;
-	UbaSessionServer = CreateSessionServer(info);
+	UbaSessionServer = SessionServer_Create(info);
 
 	CalculateKnownInputs();
 
@@ -333,14 +335,14 @@ void FUbaJobProcessor::StartUba()
 
 	if (UE::GetMultiprocessId() == 0)
 	{
-		Server_StartListen(UbaServer, uba::DefaultPort, nullptr); // Start listen so any helper on the LAN can join in
+		NetworkServer_StartListen(UbaServer, uba::DefaultPort, nullptr); // Start listen so any helper on the LAN can join in
 	}
 
 	HordeAgentManager = MakeUnique<FUbaHordeAgentManager>(ControllerModule.GetWorkingDirectory(), GetUbaBinariesPath());
 
 	auto AddClientCallback = [](void* userData, const uba::tchar* ip, uint16 port)
 		{
-			return Server_AddClient((uba::NetworkServer*)userData, ip, port, nullptr);
+			return NetworkServer_AddClient((uba::NetworkServer*)userData, ip, port, nullptr);
 		};
 	HordeAgentManager->SetAddClientCallback(AddClientCallback, UbaServer);
 
@@ -396,12 +398,12 @@ void FUbaJobProcessor::ShutDownUba()
 		return;
 	}
 
-	Server_Stop(UbaServer);
+	NetworkServer_Stop(UbaServer);
 
 	Scheduler_Destroy(UbaScheduler);
-	DestroySessionServer(UbaSessionServer);
-	DestroyStorageServer(UbaStorageServer);
-	DestroyServer(UbaServer);
+	SessionServer_Destroy(UbaSessionServer);
+	StorageServer_Destroy(UbaStorageServer);
+	NetworkServer_Destroy(UbaServer);
 
 	UbaScheduler = nullptr;
 	UbaSessionServer = nullptr;
diff --git a/Engine/Plugins/UbaController/Source/Public/UbaJobProcessor.h b/Engine/Plugins/UbaController/Source/Public/UbaJobProcessor.h
--- a/Engine/Plugins/UbaController/Source/Public/UbaJobProcessor.h
+++ b/Engine/Plugins/UbaController/Source/Public/UbaJobProcessor.h
@@ -71,7 +71,7 @@ private:
 	bool bIsWorkDone;
 
 	uba::NetworkServer* UbaServer = nullptr;
-	uba::Storage* UbaStorageServer = nullptr;
+	uba::StorageServer* UbaStorageServer = nullptr;
 	uba::SessionServer* UbaSessionServer = nullptr;
 	uba::Scheduler* UbaScheduler = nullptr;
 	TUniquePtr<FUbaHordeAgentManager> HordeAgentManager;
