[UBA] Enable UBA shader compilation by default on build machines.

- Changed "[UbaController]:Enabled" INI configuration from boolean to string, now accepting "True", "False", and "BuildMachineOnly".
- Backward compatible to commandline argument "-ini:Engine:[UbaController]:Enabled=true".
- Enable CVar "r.UbaController.DumpTraceFiles" by default so no additional configuration for "BuildMachineOnly" is required (UBA trace files don't take up a lot of memory).

#rnx

diff --git a/Engine/Plugins/UbaController/Source/Private/UbaControllerModule.cpp b/Engine/Plugins/UbaController/Source/Private/UbaControllerModule.cpp
--- a/Engine/Plugins/UbaController/Source/Private/UbaControllerModule.cpp
+++ b/Engine/Plugins/UbaController/Source/Private/UbaControllerModule.cpp
@@ -19,11 +19,11 @@ namespace UbaControllerModule
 {
 	static constexpr int32 SubFolderCount = 32;
 
-	static bool bDumpTraceFiles = false;
+	static bool bDumpTraceFiles = true;
 	static FAutoConsoleVariableRef CVarDumpTraceFiles(
 		TEXT("r.UbaController.DumpTraceFiles"),
 		bDumpTraceFiles,
-		TEXT("If true, UBA controller dumps trace files for later use with UBA visualizer in the Saved folder under UbaController"));
+		TEXT("If true, UBA controller dumps trace files for later use with UBA visualizer in the Saved folder under UbaController (Enabled by default)"));
 
 	static FString MakeAndGetDebugInfoPath()
 	{
@@ -81,7 +81,19 @@ bool FUbaControllerModule::IsSupported()
 	bool bEnabled = false;
 	if (!FParse::Param(FCommandLine::Get(), TEXT("NoUbaController")))
 	{
-		GConfig->GetBool(TEXT("UbaController"), TEXT("Enabled"), bEnabled, GEngineIni);
+		FString EnabledState;
+		GConfig->GetString(TEXT("UbaController"), TEXT("Enabled"), EnabledState, GEngineIni);
+
+		// This "Enabled" parameter is a tri-state so we have to parse it as a string.
+		// Those strings can include the comments (starting with ';') from the INI file, so we have to trim that section from the string.
+		const int32 EnabledStateCommentPosition = EnabledState.Find(TEXT(";"));
+		if (EnabledStateCommentPosition != INDEX_NONE)
+		{
+			EnabledState.RemoveAt(EnabledStateCommentPosition, EnabledState.Len() - EnabledStateCommentPosition);
+			EnabledState.RemoveSpacesInline();
+		}
+
+		bEnabled = EnabledState.Equals(TEXT("True"), ESearchCase::IgnoreCase) || (EnabledState.Equals(TEXT("BuildMachineOnly"), ESearchCase::IgnoreCase) && GIsBuildMachine);
 	}
 
 	bSupported = FPlatformProcess::SupportsMultithreading() && bEnabled;
