[UbaController]
* Moved virtual function to cpp file

diff --git a/Engine/Plugins/UbaController/Source/Private/UbaJobProcessor.cpp b/Engine/Plugins/UbaController/Source/Private/UbaJobProcessor.cpp
--- a/Engine/Plugins/UbaController/Source/Private/UbaJobProcessor.cpp
+++ b/Engine/Plugins/UbaController/Source/Private/UbaJobProcessor.cpp
@@ -483,12 +483,19 @@ uint32 FUbaJobProcessor::Run()
 		FPlatformProcess::Sleep(UbaJobProcessorOptions::SleepTimeBetweenActions);
 	}
 
+	FScopeLock lock(&bShouldProcessJobsLock);
+
 	ShutDownUba();
 
 	bIsWorkDone = true;
 	return 0;
 }
 
+void FUbaJobProcessor::Stop()
+{
+	bForceStop = true;
+};
+
 void FUbaJobProcessor::StartThread()
 {
 	Thread = FRunnableThread::Create(this, TEXT("UbaJobProcessor"), 0, TPri_SlightlyBelowNormal, FPlatformAffinity::GetPoolThreadMask());
diff --git a/Engine/Plugins/UbaController/Source/Public/UbaJobProcessor.h b/Engine/Plugins/UbaController/Source/Public/UbaJobProcessor.h
--- a/Engine/Plugins/UbaController/Source/Public/UbaJobProcessor.h
+++ b/Engine/Plugins/UbaController/Source/Public/UbaJobProcessor.h
@@ -23,7 +23,7 @@ public:
 	virtual uint32 Run() override;
 	
 	/** Aborts the main loop as soon as possible */
-	virtual void Stop() override { bForceStop = true; };
+	virtual void Stop() override;
 	
 	/** Creates the threads and starts the main loop */
 	void StartThread();
