[UBA] Improve log output for ensure condition in UBA job processor.

Use several ensure conditions instead of a single one to get filename and file size.
The failed job artifacts don't provide enough information to track down the root cause.

#rnx

diff --git a/Engine/Plugins/UbaController/Source/Private/UbaJobProcessor.cpp b/Engine/Plugins/UbaController/Source/Private/UbaJobProcessor.cpp
--- a/Engine/Plugins/UbaController/Source/Private/UbaJobProcessor.cpp
+++ b/Engine/Plugins/UbaController/Source/Private/UbaJobProcessor.cpp
@@ -509,8 +509,10 @@ bool FUbaJobProcessor::ProcessOutputFile(FTask* CompileTask)
 	IFileManager& FileManager = IFileManager::Get();
 
 	constexpr uint64 VersionAndFileSizeSize = sizeof(uint32) + sizeof(uint64);	
-	if (ensure(CompileTask && PlatformFile.FileExists(*CompileTask->CommandData.OutputFileName) &&
-		FileManager.FileSize(*CompileTask->CommandData.OutputFileName) > VersionAndFileSizeSize))
+	if (ensure(CompileTask) && 
+		ensureMsgf(PlatformFile.FileExists(*CompileTask->CommandData.OutputFileName), TEXT("Filename=%s"), *CompileTask->CommandData.OutputFileName) &&
+		ensureMsgf(FileManager.FileSize(*CompileTask->CommandData.OutputFileName) > VersionAndFileSizeSize,
+			TEXT("Filename=%s, FileSize=%d"), *CompileTask->CommandData.OutputFileName, (int32)FileManager.FileSize(*CompileTask->CommandData.OutputFileName)))
 	{
 		const TUniquePtr<FArchive> OutputFilePtr(FileManager.CreateFileReader(*CompileTask->CommandData.OutputFileName, FILEREAD_Silent));
 		if (ensure(OutputFilePtr))
