[UBA] Add command line argument "-UBA" as abbreviation to enable UbaController plugin.

This also moves all conditions to enable UbaController plugin into IsUbaControllerEnabled(), which are now "-NoUbaController", "-Uba", and the tri-state INI config "[UbaController]:Enabled".
This allows a fine-grained control over what machines support UBA due to Horde bandwidth.

#rnx

diff --git a/Engine/Plugins/UbaController/Source/Private/UbaControllerModule.cpp b/Engine/Plugins/UbaController/Source/Private/UbaControllerModule.cpp
--- a/Engine/Plugins/UbaController/Source/Private/UbaControllerModule.cpp
+++ b/Engine/Plugins/UbaController/Source/Private/UbaControllerModule.cpp
@@ -70,6 +70,40 @@ FUbaControllerModule::~FUbaControllerModule()
 	CleanWorkingDirectory();
 }
 
+static bool IsUbaControllerEnabled()
+{
+	if (FParse::Param(FCommandLine::Get(), TEXT("NoUbaController")))
+	{
+		return false;
+	}
+
+	// Check if UbaController is enabled via command line argument
+	if (FParse::Param(FCommandLine::Get(), TEXT("Uba")))
+	{
+		return true;
+	}
+
+	// Check if UbaController is enabled via INI configuration in [UbaController] section.
+	FString EnabledState;
+	GConfig->GetString(TEXT("UbaController"), TEXT("Enabled"), EnabledState, GEngineIni);
+
+	// This "Enabled" parameter is a tri-state so we have to parse it as a string.
+	// Those strings can include the comments (starting with ';') from the INI file, so we have to trim that section from the string.
+	const int32 EnabledStateCommentPosition = EnabledState.Find(TEXT(";"));
+	if (EnabledStateCommentPosition != INDEX_NONE)
+	{
+		EnabledState.RemoveAt(EnabledStateCommentPosition, EnabledState.Len() - EnabledStateCommentPosition);
+		EnabledState.RemoveSpacesInline();
+	}
+
+	if (EnabledState.Equals(TEXT("True"), ESearchCase::IgnoreCase) || (EnabledState.Equals(TEXT("BuildMachineOnly"), ESearchCase::IgnoreCase) && GIsBuildMachine))
+	{
+		return true;
+	}
+
+	return false;
+}
+
 bool FUbaControllerModule::IsSupported()
 {
 	if (bControllerInitialized)
@@ -78,23 +112,7 @@ bool FUbaControllerModule::IsSupported()
 	}
 	
 #if PLATFORM_WINDOWS
-	bool bEnabled = false;
-	if (!FParse::Param(FCommandLine::Get(), TEXT("NoUbaController")))
-	{
-		FString EnabledState;
-		GConfig->GetString(TEXT("UbaController"), TEXT("Enabled"), EnabledState, GEngineIni);
-
-		// This "Enabled" parameter is a tri-state so we have to parse it as a string.
-		// Those strings can include the comments (starting with ';') from the INI file, so we have to trim that section from the string.
-		const int32 EnabledStateCommentPosition = EnabledState.Find(TEXT(";"));
-		if (EnabledStateCommentPosition != INDEX_NONE)
-		{
-			EnabledState.RemoveAt(EnabledStateCommentPosition, EnabledState.Len() - EnabledStateCommentPosition);
-			EnabledState.RemoveSpacesInline();
-		}
-
-		bEnabled = EnabledState.Equals(TEXT("True"), ESearchCase::IgnoreCase) || (EnabledState.Equals(TEXT("BuildMachineOnly"), ESearchCase::IgnoreCase) && GIsBuildMachine);
-	}
+	const bool bEnabled = IsUbaControllerEnabled();
 
 	bSupported = FPlatformProcess::SupportsMultithreading() && bEnabled;
 #else
