[UBA] Rename trace files to reflect "MultiprocessId" to avoid confusion with cook worker ID.

#rnx

diff --git a/Engine/Plugins/UbaController/Source/Private/UbaJobProcessor.cpp b/Engine/Plugins/UbaController/Source/Private/UbaJobProcessor.cpp
--- a/Engine/Plugins/UbaController/Source/Private/UbaJobProcessor.cpp
+++ b/Engine/Plugins/UbaController/Source/Private/UbaJobProcessor.cpp
@@ -315,7 +315,7 @@ void FUbaJobProcessor::StartUba()
 	if (!ControllerModule.GetDebugInfoPath().IsEmpty())
 	{
 		static uint32 UbaSessionCounter;
-		TraceOutputFile = ControllerModule.GetDebugInfoPath() / FString::Printf(TEXT("UbaController.CookWorker-%u.Session-%u.uba"), UE::GetMultiprocessId(), UbaSessionCounter++);
+		TraceOutputFile = ControllerModule.GetDebugInfoPath() / FString::Printf(TEXT("UbaController.MultiprocessId-%u.Session-%u.uba"), UE::GetMultiprocessId(), UbaSessionCounter++);
 	}
 	info.traceOutputFile = *TraceOutputFile;
 	info.detailedTrace = UbaJobProcessorOptions::bDetailedTrace;
