[UbaController]
* Fixed so UbaControllerWorkingDir is using UE_HORDE_SHARED_DIR if set

diff --git a/Engine/Plugins/UbaController/Source/Private/UbaControllerModule.cpp b/Engine/Plugins/UbaController/Source/Private/UbaControllerModule.cpp
--- a/Engine/Plugins/UbaController/Source/Private/UbaControllerModule.cpp
+++ b/Engine/Plugins/UbaController/Source/Private/UbaControllerModule.cpp
@@ -45,13 +45,22 @@ namespace UbaControllerModule
 
 		return AbsoluteDebugInfoDirectory;
 	}
+
+	FString GetTempDir()
+	{
+		const FString HordeSharedDir = FPlatformMisc::GetEnvironmentVariable(TEXT("UE_HORDE_SHARED_DIR"));
+		if (!HordeSharedDir.IsEmpty())
+			return HordeSharedDir;
+		return FPlatformProcess::UserTempDir();
+	}
+
 };
 
 FUbaControllerModule::FUbaControllerModule()
 	: bSupported(false)
 	, bModuleInitialized(false)
 	, bControllerInitialized(false)
-	, RootWorkingDirectory(FPaths::Combine(FPlatformProcess::UserTempDir(), TEXT("UbaControllerWorkingDir")))
+	, RootWorkingDirectory(FPaths::Combine(*UbaControllerModule::GetTempDir(), TEXT("UbaControllerWorkingDir")))
 	, WorkingDirectory(FPaths::Combine(RootWorkingDirectory, FGuid::NewGuid().ToString(EGuidFormats::Digits)))
 	, NextFileID(0)
 	, NextTaskID(0)
