[UBA] Enable UbaController module for shader compilation on Mac and Linux.

- Enable UbaController module and plugin for Mac and Linux.
- Abstracts string literals for UBA as Mac/Linux use ANSICHAR while Windows uses WIDECHAR.
- Fix URL output of Horde lease (remove double slash character).

diff --git a/Engine/Plugins/UbaController/Source/Private/UbaControllerModule.cpp b/Engine/Plugins/UbaController/Source/Private/UbaControllerModule.cpp
--- a/Engine/Plugins/UbaController/Source/Private/UbaControllerModule.cpp
+++ b/Engine/Plugins/UbaController/Source/Private/UbaControllerModule.cpp
@@ -120,14 +120,9 @@ bool FUbaControllerModule::IsSupported()
 		return bSupported;
 	}
 	
-#if PLATFORM_WINDOWS
 	const bool bEnabled = IsUbaControllerEnabled();
 
 	bSupported = FPlatformProcess::SupportsMultithreading() && bEnabled;
-#else
-	// Not supported on other platforms.
-	bSupported = false;
-#endif
 	return bSupported;
 }
 
@@ -156,12 +151,20 @@ bool FUbaControllerModule::HasTasksDispatchedOrPending() const
 
 FString GetUbaBinariesPath()
 {
+#if PLATFORM_WINDOWS
 #if PLATFORM_CPU_ARM_FAMILY
 	const TCHAR* BinariesArch = TEXT("arm64");
 #else
 	const TCHAR* BinariesArch = TEXT("x64");
 #endif
-	return FPaths::Combine(FPaths::EngineDir(), "Binaries/Win64/UnrealBuildAccelerator", BinariesArch);
+	return FPaths::Combine(FPaths::EngineDir(), TEXT("Binaries"), TEXT("Win64"), TEXT("UnrealBuildAccelerator"), BinariesArch);
+#elif PLATFORM_MAC
+    return FPaths::Combine(FPaths::EngineDir(), TEXT("Binaries"), TEXT("Mac"), TEXT("UnrealBuildAccelerator"));
+#elif PLATFORM_LINUX
+    return FPaths::Combine(FPaths::EngineDir(), TEXT("Binaries"), TEXT("Linux"), TEXT("UnrealBuildAccelerator"));
+#else
+#error Unsupported platform to compile UbaController plugin. Only Win64, Mac, and Linux are supported!
+#endif
 }
 
 void FUbaControllerModule::LoadDependencies()
diff --git a/Engine/Plugins/UbaController/Source/Private/UbaJobProcessor.cpp b/Engine/Plugins/UbaController/Source/Private/UbaJobProcessor.cpp
--- a/Engine/Plugins/UbaController/Source/Private/UbaJobProcessor.cpp
+++ b/Engine/Plugins/UbaController/Source/Private/UbaJobProcessor.cpp
@@ -6,12 +6,16 @@
 #include "Interfaces/ITargetPlatformManagerModule.h"
 #include "Misc/ConfigCacheIni.h"
 #include "Misc/CoreMisc.h"
+#include "UbaControllerModule.h"
 #include "UbaHordeAgentManager.h"
 #include "UbaProcessStartInfo.h"
 #include "UbaSessionServerCreateInfo.h"
+#include "UbaStringConversion.h"
+
+#if PLATFORM_WINDOWS
 #include "Windows/AllowWindowsPlatformTypes.h"
-#include "UbaControllerModule.h"
 #include "Windows/HideWindowsPlatformTypes.h"
+#endif
 
 #include "HAL/FileManager.h"
 #include "HAL/PlatformFileManager.h"
@@ -120,26 +124,26 @@ FUbaJobProcessor::FUbaJobProcessor(
 	LastTimeCheckedForTasks(0),
 	bShouldProcessJobs(false),
 	bIsWorkDone(false),
-	LogWriter([]() {}, []() {}, [](uba::LogEntryType type, const wchar_t* str, uba::u32 strlen)
+	LogWriter([]() {}, []() {}, [](uba::LogEntryType type, const uba::tchar* str, uba::u32 /*strlen*/)
 	{
 			switch (type)
 			{
 			case uba::LogEntryType_Error:
-				UE_LOG(LogUbaController, Error, TEXT("%s"), str);
+				UE_LOG(LogUbaController, Error, TEXT("%s"), UBASTRING_TO_TCHAR(str));
 				break;
 			case uba::LogEntryType_Warning:
-				UE_LOG(LogUbaController, Warning, TEXT("%s"), str);
+				UE_LOG(LogUbaController, Warning, TEXT("%s"), UBASTRING_TO_TCHAR(str));
 				break;
 			case uba::LogEntryType_Info:
 				if (UbaJobProcessorOptions::UbaLogVerbosity >= UbaJobProcessorOptions::UbaLogVerbosity_High)
 				{
-					UE_LOG(LogUbaController, Display, TEXT("%s"), str);
+					UE_LOG(LogUbaController, Display, TEXT("%s"), UBASTRING_TO_TCHAR(str));
 				}
 				break;
 			default:
 				if (UbaJobProcessorOptions::UbaLogVerbosity >= UbaJobProcessorOptions::UbaLogVerbosity_Max)
 				{
-					UE_LOG(LogUbaController, Display, TEXT("%s"), str);
+					UE_LOG(LogUbaController, Display, TEXT("%s"), UBASTRING_TO_TCHAR(str));
 				}
 				break;
 			}
@@ -147,7 +151,7 @@ FUbaJobProcessor::FUbaJobProcessor(
 {
 	Uba_SetCustomAssertHandler([](const uba::tchar* text)
 		{
-			checkf(false, TEXT("%s"), text);
+			checkf(false, TEXT("%s"), UBASTRING_TO_TCHAR(text));
 		});
 
 	if (!GConfig->GetInt(TEXT("UbaController"), TEXT("MaxLocalParallelJobs"), MaxLocalParallelJobs, GEngineIni))
@@ -231,23 +235,28 @@ void FUbaJobProcessor::CalculateKnownInputs()
 void FUbaJobProcessor::RunTaskWithUba(FTask* Task)
 {
 	FTaskCommandData& Data = Task->CommandData;
-	SessionServer_RegisterNewFile(UbaSessionServer, *Data.InputFileName);
+	SessionServer_RegisterNewFile(UbaSessionServer, TCHAR_TO_UBASTRING(*Data.InputFileName));
 
 	FString InputFileName = FPaths::GetCleanFilename(Data.InputFileName);
 	FString OutputFileName = FPaths::GetCleanFilename(Data.OutputFileName);
 	FString Parameters = FString::Printf(TEXT("\"%s/\" %d 0 \"%s\" \"%s\" %s "), *Data.WorkingDirectory, Data.DispatcherPID, *InputFileName, *OutputFileName, *Data.ExtraCommandArgs);
 	FString AppDir = FPaths::GetPath(Data.Command);
+	
+	FStringToUbaStringConversion UbaCommandStr(*Data.Command);
+	FStringToUbaStringConversion UbaParametersStr(*Parameters);
+	FStringToUbaStringConversion UbaInputFileNameStr(*InputFileName);
+	FStringToUbaStringConversion UbaWorkingDirStr(*AppDir);
 
 	uba::ProcessStartInfo ProcessInfo;
-	ProcessInfo.application = *Data.Command;
-	ProcessInfo.arguments = *Parameters;
-	ProcessInfo.description = *InputFileName;
-	ProcessInfo.workingDir = *AppDir;
+	ProcessInfo.application = UbaCommandStr.Get();
+	ProcessInfo.arguments = UbaParametersStr.Get();
+	ProcessInfo.description = UbaInputFileNameStr.Get();
+	ProcessInfo.workingDir = UbaWorkingDirStr.Get();
 	ProcessInfo.writeOutputFilesOnFail = true;
 
 	if (UbaJobProcessorOptions::bProcessLogEnabled)
 	{
-		ProcessInfo.logFile = *InputFileName;
+		ProcessInfo.logFile = UbaInputFileNameStr.Get();
 	}
 	
 	struct ExitedInfo
@@ -270,17 +279,17 @@ void FUbaJobProcessor::RunTaskWithUba(FTask* Task)
 			uint32 logLineIndex = 0;
 			while (const uba::tchar* line = ProcessHandle_GetLogLine(&ph, logLineIndex++))
 			{
-				UE_LOG(LogUbaController, Display, TEXT("%s"), line);
+				UE_LOG(LogUbaController, Display, TEXT("%s"), UBASTRING_TO_TCHAR(line));
 			}
 
 			if (auto Info = (ExitedInfo*)userData) // It can be null if custom message has already handled all of them
 			{
 				IFileManager::Get().Delete(*Info->InputFile);
-				SessionServer_RegisterDeleteFile(Info->Processor->UbaSessionServer, *Info->InputFile);
+				SessionServer_RegisterDeleteFile(Info->Processor->UbaSessionServer, TCHAR_TO_UBASTRING(*Info->InputFile));
 				Info->Processor->HandleUbaJobFinished(Info->Task);
 
-				StorageServer_DeleteFile(Info->Processor->UbaStorageServer, *Info->InputFile);
-				StorageServer_DeleteFile(Info->Processor->UbaStorageServer, *Info->OutputFile);
+				StorageServer_DeleteFile(Info->Processor->UbaStorageServer, TCHAR_TO_UBASTRING(*Info->InputFile));
+				StorageServer_DeleteFile(Info->Processor->UbaStorageServer, TCHAR_TO_UBASTRING(*Info->OutputFile));
 
 				delete Info;
 			}
@@ -302,11 +311,12 @@ void FUbaJobProcessor::StartUba()
 	IFileManager::Get().MakeDirectory(*RootDir, true);
 
 	uba::u64 casCapacityBytes = 32llu * 1024 * 1024 * 1024;
-	UbaStorageServer = StorageServer_Create(*UbaServer, *RootDir, casCapacityBytes, true, LogWriter);
+	UbaStorageServer = StorageServer_Create(*UbaServer, TCHAR_TO_UBASTRING(*RootDir), casCapacityBytes, true, LogWriter);
 
 	uba::SessionServerCreateInfo info(*(uba::Storage*)UbaStorageServer, *UbaServer, LogWriter);
 	info.launchVisualizer = UbaJobProcessorOptions::bAutoLaunchVisualizer;
-	info.rootDir = *RootDir;
+	const FStringToUbaStringConversion UbaRootDirStr(*RootDir);
+	info.rootDir = UbaRootDirStr.Get();
 	info.allowMemoryMaps = false; // Skip using memory maps
 	info.remoteLogEnabled = UbaJobProcessorOptions::bProcessLogEnabled;
 
@@ -317,10 +327,12 @@ void FUbaJobProcessor::StartUba()
 		static uint32 UbaSessionCounter;
 		TraceOutputFile = ControllerModule.GetDebugInfoPath() / FString::Printf(TEXT("UbaController.MultiprocessId-%u.Session-%u.uba"), UE::GetMultiprocessId(), UbaSessionCounter++);
 	}
-	info.traceOutputFile = *TraceOutputFile;
+	const FStringToUbaStringConversion UbaTraceOutputFileStr(*TraceOutputFile);
+	info.traceOutputFile = UbaTraceOutputFileStr.Get();
 	info.detailedTrace = UbaJobProcessorOptions::bDetailedTrace;
 	FString TraceName = FString::Printf(TEXT("UbaController_%s"), *FGuid::NewGuid().ToString(EGuidFormats::Digits));
-	info.traceName = *TraceName;
+	const FStringToUbaStringConversion UbaTraceNameStr(*TraceName);
+	info.traceName = UbaTraceNameStr.Get();
 
 
 	//info.remoteLogEnabled = true;
@@ -546,7 +558,7 @@ bool FUbaJobProcessor::ProcessOutputFile(FTask* CompileTask)
 	}
 	else
 	{
-		const FString OutputFileName = CompileTask != nullptr ? CompileTask->CommandData.OutputFileName : TEXT("Invalid CompileTask, cannot retrieve name");
+		const FString OutputFileName = CompileTask != nullptr ? *CompileTask->CommandData.OutputFileName : TEXT("Invalid CompileTask, cannot retrieve name");
 		UE_LOG(LogUbaController, Error, TEXT("Output File [%s] is invalid or does not exist"), *OutputFileName);
 		return false;
 	}
diff --git a/Engine/Plugins/UbaController/Source/Private/UbaStringConversion.h b/Engine/Plugins/UbaController/Source/Private/UbaStringConversion.h
new file mode 100644
--- /dev/null
+++ b/Engine/Plugins/UbaController/Source/Private/UbaStringConversion.h
@@ -0,0 +1,13 @@
+// Copyright Epic Games, Inc. All Rights Reserved.
+#pragma once
+
+#include "CoreTypes.h"
+#include "Containers/StringConv.h"
+#include "UbaExports.h"
+
+using FStringToUbaStringConversion = TStringConversion<TStringConvert<TCHAR, uba::tchar>>;
+using FUbaStringToStringConversion = TStringConversion<TStringConvert<uba::tchar, TCHAR>>;
+
+#define TCHAR_TO_UBASTRING(STR) FStringToUbaStringConversion(STR).Get()
+#define UBASTRING_TO_TCHAR(STR) FUbaStringToStringConversion(STR).Get()
+
diff --git a/Engine/Plugins/UbaController/Source/UbaController.build.cs b/Engine/Plugins/UbaController/Source/UbaController.build.cs
--- a/Engine/Plugins/UbaController/Source/UbaController.build.cs
+++ b/Engine/Plugins/UbaController/Source/UbaController.build.cs
@@ -30,11 +30,24 @@ public class UbaController : ModuleRules
         PublicIncludePaths.Add(Path.Combine(UbaSourcePath, "Host", "Public"));
         PublicIncludePathModuleNames.Add("BLAKE3");
 
-		string UbaBinariesPath = Path.Combine(EngineDirectory, "Binaries", "Win64", "UnrealBuildAccelerator", Arch);
-
-		PublicAdditionalLibraries.Add(Path.Combine(UbaBinariesPath, "UbaHost.lib"));
-		PrivateRuntimeLibraryPaths.Add(Path.Combine(ModuleDirectory, "lib"));
+		string UbaBinariesPath = Path.Combine(EngineDirectory, "Binaries", Target.Platform.ToString(), "UnrealBuildAccelerator");
+		
+		if (Target.Platform == UnrealTargetPlatform.Win64)
+		{
+			UbaBinariesPath = Path.Combine(UbaBinariesPath, Arch);
+			
+			PublicAdditionalLibraries.Add(Path.Combine(UbaBinariesPath, "UbaHost.lib"));
+			PrivateRuntimeLibraryPaths.Add(Path.Combine(ModuleDirectory, "lib"));
 
-		PublicDelayLoadDLLs.Add("UbaHost.dll");
+			PublicDelayLoadDLLs.Add("UbaHost.dll");
+		}
+		else if (Target.Platform == UnrealTargetPlatform.Mac)
+		{
+			PublicDelayLoadDLLs.Add(Path.Combine(UbaBinariesPath, "libUbaHost.dylib"));
+		}
+		else if (Target.Platform == UnrealTargetPlatform.Linux)
+		{
+			PublicDelayLoadDLLs.Add(Path.Combine(UbaBinariesPath, "libUbaHost.so"));
+		}
 	}
-}
\ No newline at end of file
+}
diff --git a/Engine/Source/Developer/UbaCoordinatorHorde/Private/UbaHordeAgentManager.cpp b/Engine/Source/Developer/UbaCoordinatorHorde/Private/UbaHordeAgentManager.cpp
--- a/Engine/Source/Developer/UbaCoordinatorHorde/Private/UbaHordeAgentManager.cpp
+++ b/Engine/Source/Developer/UbaCoordinatorHorde/Private/UbaHordeAgentManager.cpp
@@ -202,19 +202,27 @@ void FUbaHordeAgentManager::ThreadAgent(FHordeAgentWrapper& Wrapper)
 		FScopeLock ScopeLock(&BundleRefPathsLock);
 		if (BundleRefPaths.IsEmpty())
 		{
-			struct BundleRec { const char* File; const TCHAR* BundleRef; };
-			BundleRec BundleRecs[] =
+			struct FBundleRec
 			{
-				{ AppName, TEXT("UbaAgent.Bundle.ref") },
-#if PLATFORM_LINUX
-				{ "UbaAgent.debug", TEXT("UbaAgent.debug.Bundle.ref") },
-				{ "libclang_rt.tsan.so", TEXT("Tsan.Bundle.ref") },
+				const TCHAR* Filename;
+				const TCHAR* BundleRef;
+			};
+			const FBundleRec BundleRecs[] =
+			{
+#if PLATFORM_WINDOWS
+				{ TEXT("UbaAgent.exe"), TEXT("UbaAgent.Bundle.ref") },
+#elif PLATFORM_LINUX
+				{ TEXT("UbaAgent"), TEXT("UbaAgent.Bundle.ref") },
+				{ TEXT("UbaAgent.debug"), TEXT("UbaAgent.debug.Bundle.ref") },
+//				{ TEXT("libclang_rt.tsan.so"), TEXT("Tsan.Bundle.ref") }, // for debugging
+#elif PLATFORM_MAC
+				{ TEXT("UbaAgent"), TEXT("UbaAgent.Bundle.ref") },
 #endif
 			};
 
-			for (const BundleRec& Rec : BundleRecs)
+			for (const FBundleRec& Rec : BundleRecs)
 			{
-				const FString FilePath = FPaths::Combine(BinariesPath, ANSI_TO_TCHAR(Rec.File));
+				const FString FilePath = FPaths::Combine(BinariesPath, Rec.Filename);
 				FString BundlePath = FPaths::Combine(WorkingDir, Rec.BundleRef);
 
 				if (!CreateHordeBundleFromFile(*FilePath, *BundlePath))
diff --git a/Engine/Source/Developer/UbaCoordinatorHorde/Private/UbaHordeMetaClient.cpp b/Engine/Source/Developer/UbaCoordinatorHorde/Private/UbaHordeMetaClient.cpp
--- a/Engine/Source/Developer/UbaCoordinatorHorde/Private/UbaHordeMetaClient.cpp
+++ b/Engine/Source/Developer/UbaCoordinatorHorde/Private/UbaHordeMetaClient.cpp
@@ -124,7 +124,7 @@ TSharedPtr<FUbaHordeMetaClient::HordeMachinePromise, ESPMode::ThreadSafe> FUbaHo
 
 			if (TSharedPtr<FJsonValue> LeaseIdValue = OutJson->AsObject()->TryGetField(TEXT("leaseId")))
 			{
-				const FString AgentWebPortalUrl = FString::Format(TEXT("{0}/lease/{1}"), { this->ServerUrl, LeaseIdValue->AsString() });
+				const FString AgentWebPortalUrl = FString::Format(TEXT("{0}lease/{1}"), { this->ServerUrl, LeaseIdValue->AsString() });
 				UE_LOG(
 					LogUbaHorde, Display, TEXT("UBA Horde machine assigned (%s) [%s:%u]: %s"),
 					(Info.bRunsWindowOS ? TEXT("Microsoft Windows") : TEXT("POSIX/Wine")), *IpString, (uint32)PortNumber, *AgentWebPortalUrl
