[UBAController]
* Fixed known inputs code so it works on non windows

diff --git a/Engine/Plugins/UbaController/Source/Private/UbaJobProcessor.cpp b/Engine/Plugins/UbaController/Source/Private/UbaJobProcessor.cpp
--- a/Engine/Plugins/UbaController/Source/Private/UbaJobProcessor.cpp
+++ b/Engine/Plugins/UbaController/Source/Private/UbaJobProcessor.cpp
@@ -183,20 +183,36 @@ void FUbaJobProcessor::CalculateKnownInputs()
 
 	auto AddKnownInput = [&](const FString& file)
 		{
+			UE_LOG(LogUbaController, Display, TEXT("ADDING KNOWN FILE: %s"), *file);
+			#if PLATFORM_WINDOWS
 			auto& fileData = file.GetCharArray();
+			const uba::tchar* fileName = fileData.GetData();
+			size_t fileNameLen = fileData.Num();
+			#else
+			FStringToUbaStringConversion conv(*file);
+			const uba::tchar* fileName = conv.Get();
+			size_t fileNameLen = strlen(fileName) + 1;
+			#endif
 			auto num = KnownInputsBuffer.Num();
-			KnownInputsBuffer.SetNum(num + fileData.Num());
-			memcpy(KnownInputsBuffer.GetData() + num, fileData.GetData(), fileData.Num() * sizeof(uba::tchar));
+			KnownInputsBuffer.SetNum(num + fileNameLen);
+			memcpy(KnownInputsBuffer.GetData() + num, fileName, fileNameLen * sizeof(uba::tchar));
 			++KnownInputsCount;
 		};
 
 	// Get the binaries
 	TArray<FString> KnownFileNames;
-	FString BinDir = FPaths::Combine(FPaths::EngineDir(), TEXT("Binaries/Win64")); // TODO: Need to support other targets than Win64
-	IFileManager::Get().FindFilesRecursive(KnownFileNames, *BinDir, TEXT("ShaderCompileWorker*.*"), true, false);
+	FString BinDir = FPaths::Combine(FPaths::EngineDir(), TEXT("Binaries"), FPlatformProcess::GetBinariesSubdirectory());
+
+	#if PLATFORM_WINDOWS
+	AddKnownInput(*FPaths::Combine(BinDir, TEXT("ShaderCompileWorker.exe")));
+	#else
+	AddKnownInput(*FPaths::Combine(BinDir, TEXT("ShaderCompileWorker")));
+	#endif
+
+	IFileManager::Get().FindFilesRecursive(KnownFileNames, *BinDir, TEXT("ShaderCompileWorker-*.*"), true, false);
 	for (const FString& file : KnownFileNames)
 	{
-		if (!file.EndsWith(TEXT(".pdb")))
+		if (file.EndsWith(FPlatformProcess::GetModuleExtension()))
 		{
 			AddKnownInput(file);
 		}
