[UbaController]
* Removed accidental log line

diff --git a/Engine/Plugins/UbaController/Source/Private/UbaJobProcessor.cpp b/Engine/Plugins/UbaController/Source/Private/UbaJobProcessor.cpp
--- a/Engine/Plugins/UbaController/Source/Private/UbaJobProcessor.cpp
+++ b/Engine/Plugins/UbaController/Source/Private/UbaJobProcessor.cpp
@@ -183,7 +183,6 @@ void FUbaJobProcessor::CalculateKnownInputs()
 
 	auto AddKnownInput = [&](const FString& file)
 		{
-			UE_LOG(LogUbaController, Display, TEXT("ADDING KNOWN FILE: %s"), *file);
 			#if PLATFORM_WINDOWS
 			auto& fileData = file.GetCharArray();
 			const uba::tchar* fileName = fileData.GetData();
