[UbaController]
* Added support for "MacPool" in config for horde.. so a different agent pool can be used for mac

diff --git a/Engine/Plugins/UbaController/Source/Private/UbaJobProcessor.cpp b/Engine/Plugins/UbaController/Source/Private/UbaJobProcessor.cpp
--- a/Engine/Plugins/UbaController/Source/Private/UbaJobProcessor.cpp
+++ b/Engine/Plugins/UbaController/Source/Private/UbaJobProcessor.cpp
@@ -382,7 +382,18 @@ void FUbaJobProcessor::StartUba()
 		HordeConfig.RemoveFromEnd(TEXT(")"));
 
 		FString Pool;
-		if (FParse::Value(*HordeConfig, TEXT("Pool="), Pool))
+		bool FoundPool = false;
+
+		#if PLATFORM_MAC
+		FoundPool = FParse::Value(*HordeConfig, TEXT("MacPool="), Pool);
+		#endif
+
+		if (!FoundPool)
+		{
+			FoundPool = FParse::Value(*HordeConfig, TEXT("Pool="), Pool);
+		}
+
+		if (FoundPool)
 		{
 			UE_LOG(LogUbaController, Log, TEXT("Found UBA controller Pool: \"%s\""), *Pool);
 			HordeAgentManager->SetPool(Pool);
