UBAController: Make StorageDir obey existing root paths

#rnx

diff --git a/Engine/Plugins/UbaController/Source/Private/UbaControllerModule.cpp b/Engine/Plugins/UbaController/Source/Private/UbaControllerModule.cpp
--- a/Engine/Plugins/UbaController/Source/Private/UbaControllerModule.cpp
+++ b/Engine/Plugins/UbaController/Source/Private/UbaControllerModule.cpp
@@ -46,12 +46,12 @@ namespace UbaControllerModule
 		return AbsoluteDebugInfoDirectory;
 	}
 
-	FString GetTempDir()
+	static const FString GetTempDir()
 	{
-		const FString HordeSharedDir = FPlatformMisc::GetEnvironmentVariable(TEXT("UE_HORDE_SHARED_DIR"));
-		if (!HordeSharedDir.IsEmpty())
-			return HordeSharedDir;
-		return FPlatformProcess::UserTempDir();
+		static FString HordeSharedDir = FPlatformMisc::GetEnvironmentVariable(TEXT("UE_HORDE_SHARED_DIR"));
+		if (HordeSharedDir.IsEmpty())
+			HordeSharedDir = FPlatformProcess::UserTempDir();
+		return HordeSharedDir;
 	}
 
 };
@@ -67,6 +67,11 @@ FUbaControllerModule::FUbaControllerModule()
 {
 }
 
+FString FUbaControllerModule::GetTempDir()
+{
+	return UbaControllerModule::GetTempDir();
+}
+
 FUbaControllerModule::~FUbaControllerModule()
 {	
 	if (JobDispatcherThread)
diff --git a/Engine/Plugins/UbaController/Source/Private/UbaJobProcessor.cpp b/Engine/Plugins/UbaController/Source/Private/UbaJobProcessor.cpp
--- a/Engine/Plugins/UbaController/Source/Private/UbaJobProcessor.cpp
+++ b/Engine/Plugins/UbaController/Source/Private/UbaJobProcessor.cpp
@@ -324,7 +324,7 @@ void FUbaJobProcessor::StartUba()
 
 	UbaServer = NetworkServer_Create(LogWriter);
 
-	FString RootDir = FString::Printf(TEXT("%s/%s/%u"), FPlatformProcess::UserTempDir(), TEXT("UbaControllerStorageDir"), UE::GetMultiprocessId());
+	FString RootDir = FString::Printf(TEXT("%s/%s/%u"), *FUbaControllerModule::GetTempDir(), TEXT("UbaControllerStorageDir"), UE::GetMultiprocessId());
 	IFileManager::Get().MakeDirectory(*RootDir, true);
 
 	uba::u64 casCapacityBytes = 32llu * 1024 * 1024 * 1024;
diff --git a/Engine/Plugins/UbaController/Source/Public/UbaControllerModule.h b/Engine/Plugins/UbaController/Source/Public/UbaControllerModule.h
--- a/Engine/Plugins/UbaController/Source/Public/UbaControllerModule.h
+++ b/Engine/Plugins/UbaController/Source/Public/UbaControllerModule.h
@@ -40,6 +40,8 @@ public:
 	// Queue of tasks submitted by the engine, but not yet dispatched to the controller.
 	TQueue<FTask*, EQueueMode::Mpsc> PendingRequestedCompilationTasks;
 
+	static FString GetTempDir();
+
 private:
 	void LoadDependencies();
 
