[UBAController] Fix editor crash when cooking project with all online-related module disabled.

diff --git a/Engine/Plugins/UbaController/Source/Private/UbaControllerModule.cpp b/Engine/Plugins/UbaController/Source/Private/UbaControllerModule.cpp
--- a/Engine/Plugins/UbaController/Source/Private/UbaControllerModule.cpp
+++ b/Engine/Plugins/UbaController/Source/Private/UbaControllerModule.cpp
@@ -1,7 +1,7 @@
 // Copyright Epic Games, Inc. All Rights Reserved.
 
 #include "UbaControllerModule.h"
-
+#include "HttpModule.h"
 #include "Misc/ConfigCacheIni.h"
 #include "UbaJobProcessor.h"
 #include "Features/IModularFeatures.h"
@@ -255,7 +255,8 @@ void FUbaControllerModule::InitializeController()
 			{
 				DebugInfoPath = UbaControllerModule::MakeAndGetDebugInfoPath();
 			}
-
+            // Make sure http module is loaded in game thread before launch UBA client
+			FModuleManager::LoadModuleChecked<FHttpModule>("HTTP");
 			JobDispatcherThread = MakeShared<FUbaJobProcessor>(*this);
 			JobDispatcherThread->StartThread();
 		}
