[UBA] Handle error of missing file in UbaController as log entry instead of ensure condition.

Instead of crashing with a failed ensure-condition when a UBA output file is missing, we handle it in the else-branch and report an error to avoid a crash.
This is the same behavior as in the FASTBuild controller.

#rnx

diff --git a/Engine/Plugins/UbaController/Source/Private/UbaJobProcessor.cpp b/Engine/Plugins/UbaController/Source/Private/UbaJobProcessor.cpp
--- a/Engine/Plugins/UbaController/Source/Private/UbaJobProcessor.cpp
+++ b/Engine/Plugins/UbaController/Source/Private/UbaJobProcessor.cpp
@@ -519,9 +519,8 @@ bool FUbaJobProcessor::ProcessOutputFile(FDistributedBuildTask* CompileTask)
 
 	constexpr uint64 VersionAndFileSizeSize = sizeof(uint32) + sizeof(uint64);	
 	if (ensure(CompileTask) && 
-		ensureMsgf(PlatformFile.FileExists(*CompileTask->CommandData.OutputFileName), TEXT("Filename=%s"), *CompileTask->CommandData.OutputFileName) &&
-		ensureMsgf(FileManager.FileSize(*CompileTask->CommandData.OutputFileName) > VersionAndFileSizeSize,
-			TEXT("Filename=%s, FileSize=%d"), *CompileTask->CommandData.OutputFileName, (int32)FileManager.FileSize(*CompileTask->CommandData.OutputFileName)))
+		PlatformFile.FileExists(*CompileTask->CommandData.OutputFileName) &&
+		FileManager.FileSize(*CompileTask->CommandData.OutputFileName) > VersionAndFileSizeSize)
 	{
 		const TUniquePtr<FArchive> OutputFilePtr(FileManager.CreateFileReader(*CompileTask->CommandData.OutputFileName, FILEREAD_Silent));
 		if (ensure(OutputFilePtr))
