[UbaController]
* Changed code to use configs to set properties instead of relying on structs directly. This will make this code more resilient to changes since we submit binaries of uba and can easily get incompatible struct layouts

diff --git a/Engine/Plugins/UbaController/Source/Private/UbaJobProcessor.cpp b/Engine/Plugins/UbaController/Source/Private/UbaJobProcessor.cpp
--- a/Engine/Plugins/UbaController/Source/Private/UbaJobProcessor.cpp
+++ b/Engine/Plugins/UbaController/Source/Private/UbaJobProcessor.cpp
@@ -9,8 +9,6 @@
 #include "UbaControllerModule.h"
 #include "UbaHordeAgentManager.h"
 #include "UbaHordeConfig.h"
-#include "UbaProcessStartInfo.h"
-#include "UbaSessionServerCreateInfo.h"
 #include "UbaStringConversion.h"
 
 #if PLATFORM_WINDOWS
@@ -272,25 +270,6 @@ void FUbaJobProcessor::RunTaskWithUba(FDistributedBuildTask* Task)
 	FString Parameters = FString::Printf(TEXT("\"%s/\" %d 0 \"%s\" \"%s\" %s "), *Data.WorkingDirectory, Data.DispatcherPID, *InputFileName, *OutputFileName, *Data.ExtraCommandArgs);
 	FString AppDir = FPaths::GetPath(Data.Command);
 	
-	FStringToUbaStringConversion UbaCommandStr(*Data.Command);
-	FStringToUbaStringConversion UbaParametersStr(*Parameters);
-	FStringToUbaStringConversion UbaInputFileNameStr(*InputFileName);
-	FStringToUbaStringConversion UbaWorkingDirStr(*AppDir);
-	FStringToUbaStringConversion UbaTaskDescription(*Data.Description);
-
-	uba::ProcessStartInfo ProcessInfo;
-	ProcessInfo.application = UbaCommandStr.Get();
-	ProcessInfo.arguments = UbaParametersStr.Get();
-	ProcessInfo.description = UbaInputFileNameStr.Get();
-	ProcessInfo.workingDir = UbaWorkingDirStr.Get();
-	ProcessInfo.writeOutputFilesOnFail = true;
-	ProcessInfo.breadcrumbs = UbaTaskDescription.Get();
-
-	if (UbaJobProcessorOptions::bProcessLogEnabled)
-	{
-		ProcessInfo.logFile = UbaInputFileNameStr.Get();
-	}
-	
 	struct ExitedInfo
 	{
 		FUbaJobProcessor* Processor;
@@ -305,8 +284,7 @@ void FUbaJobProcessor::RunTaskWithUba(FDistributedBuildTask* Task)
 	Info->OutputFile = Data.OutputFileName;
 	Info->Task = Task;
 
-	ProcessInfo.userData = Info;
-	ProcessInfo.exitedFunc = [](void* userData, const uba::ProcessHandle& ph)
+	static auto ExitedFunc = [](void* userData, const uba::ProcessHandle& ph)
 		{
 			uint32 logLineIndex = 0;
 			while (const uba::tchar* line = ProcessHandle_GetLogLine(&ph, logLineIndex++))
@@ -327,8 +305,29 @@ void FUbaJobProcessor::RunTaskWithUba(FDistributedBuildTask* Task)
 			}
 		};
 
+	FStringToUbaStringConversion UbaInputFileNameStr(*InputFileName);
+
+	uba::Config* Config = Config_Create();
+	uba::ConfigTable* RootTable = Config_RootTable(*Config);
+	ConfigTable_AddValueString(*RootTable, TC("Application"), FStringToUbaStringConversion(*Data.Command).Get());
+	ConfigTable_AddValueString(*RootTable, TC("Arguments"), FStringToUbaStringConversion(*Parameters).Get());
+	ConfigTable_AddValueString(*RootTable, TC("Description"), UbaInputFileNameStr.Get());
+	ConfigTable_AddValueString(*RootTable, TC("WorkingDir"), FStringToUbaStringConversion(*AppDir).Get());
+	ConfigTable_AddValueString(*RootTable, TC("Breadcrumbs"), FStringToUbaStringConversion(*Data.Description).Get());
+	ConfigTable_AddValueBool(*RootTable,   TC("WriteOutputFilesOnFail"), true);
+
+	if (UbaJobProcessorOptions::bProcessLogEnabled)
+	{
+		ConfigTable_AddValueString(*RootTable, TC("LogFile"), UbaInputFileNameStr.Get());
+	}
+
+	uba::ProcessStartInfo* StartInfo = ProcessStartInfo_Create3(*Config);
+	ProcessStartInfo_SetExitedCallback(*StartInfo, ExitedFunc, Info);
 
-	Scheduler_EnqueueProcess(UbaScheduler, ProcessInfo, 1.0f, KnownInputsBuffer.GetData(), KnownInputsBuffer.Num()*sizeof(uba::tchar), KnownInputsCount);
+	Scheduler_EnqueueProcess(UbaScheduler, *StartInfo, 1.0f, KnownInputsBuffer.GetData(), KnownInputsBuffer.Num()*sizeof(uba::tchar), KnownInputsCount);
+
+	ProcessStartInfo_Destroy(StartInfo);
+	Config_Destroy(Config);
 }
 
 FString GetUbaBinariesPath();
@@ -337,43 +336,59 @@ void FUbaJobProcessor::StartUba()
 {
 	checkf(UbaServer == nullptr, TEXT("FUbaJobProcessor::StartUba() was called twice before FUbaJobProcessor::ShutDownUba()"));
 
-	UbaServer = NetworkServer_Create(LogWriter);
-
 	FString RootDir = FString::Printf(TEXT("%s/%s/%u"), *FUbaControllerModule::GetTempDir(), TEXT("UbaControllerStorageDir"), UE::GetMultiprocessId());
 	IFileManager::Get().MakeDirectory(*RootDir, true);
 
-	uba::u64 casCapacityBytes = 32llu * 1024 * 1024 * 1024;
-	UbaStorageServer = StorageServer_Create(*UbaServer, TCHAR_TO_UBASTRING(*RootDir), casCapacityBytes, true, LogWriter);
-
-	uba::SessionServerCreateInfo info(*(uba::Storage*)UbaStorageServer, *UbaServer, LogWriter);
-	info.launchVisualizer = UbaJobProcessorOptions::bAutoLaunchVisualizer;
-	const FStringToUbaStringConversion UbaRootDirStr(*RootDir);
-	info.rootDir = UbaRootDirStr.Get();
-	info.allowMemoryMaps = false; // Skip using memory maps
-	info.remoteLogEnabled = UbaJobProcessorOptions::bProcessLogEnabled;
-
-	info.traceEnabled = true;
 	FString TraceOutputFile;
 	if (!ControllerModule.GetDebugInfoPath().IsEmpty())
 	{
 		static uint32 UbaSessionCounter;
 		TraceOutputFile = ControllerModule.GetDebugInfoPath() / FString::Printf(TEXT("UbaController.MultiprocessId-%u.Session-%u.uba"), UE::GetMultiprocessId(), UbaSessionCounter++);
 	}
-	const FStringToUbaStringConversion UbaTraceOutputFileStr(*TraceOutputFile);
-	info.traceOutputFile = UbaTraceOutputFileStr.Get();
-	info.detailedTrace = UbaJobProcessorOptions::bDetailedTrace;
 	FString TraceName = FString::Printf(TEXT("UbaController_%s"), *FGuid::NewGuid().ToString(EGuidFormats::Digits));
-	const FStringToUbaStringConversion UbaTraceNameStr(*TraceName);
-	info.traceName = UbaTraceNameStr.Get();
 
+	uba::Config* Config = Config_Create();
+	{
+		uba::ConfigTable* RootTable = Config_RootTable(*Config);
+		ConfigTable_AddValueString(*RootTable,     TC("RootDir"), TCHAR_TO_UBASTRING(*RootDir));
+
+		uba::ConfigTable* StorageTable = Config_AddTable(*Config, TC("Storage"));
+		ConfigTable_AddValueU64(*StorageTable,     TC("CasCapacityBytes"), 32llu * 1024 * 1024 * 1024);
+
+		uba::ConfigTable* SessionTable = Config_AddTable(*Config, TC("Session"));
+		ConfigTable_AddValueBool(*SessionTable,    TC("LaunchVisualizer"), UbaJobProcessorOptions::bAutoLaunchVisualizer);
+		ConfigTable_AddValueBool(*SessionTable,    TC("AllowMemoryMaps"), false); // Skip using memory maps
+		ConfigTable_AddValueBool(*SessionTable,    TC("RemoteLogEnabled"), UbaJobProcessorOptions::bProcessLogEnabled);
+		ConfigTable_AddValueBool(*SessionTable,    TC("TraceEnabled"), true);
+		ConfigTable_AddValueString(*SessionTable,  TC("TraceOutputFile"), TCHAR_TO_UBASTRING(*TraceOutputFile));
+		ConfigTable_AddValueBool(*SessionTable,    TC("DetailedTrace"), UbaJobProcessorOptions::bDetailedTrace);
+		ConfigTable_AddValueString(*SessionTable,  TC("TraceName"), TCHAR_TO_UBASTRING(*TraceName));
+
+		uba::ConfigTable* SchedulerTable = Config_AddTable(*Config, TC("Scheduler"));
+		ConfigTable_AddValueU32(*SchedulerTable,   TC("MaxLocalProcessors"), MaxLocalParallelJobs);
+		ConfigTable_AddValueBool(*SchedulerTable,  TC("EnableProcessReuse"), UbaJobProcessorOptions::bAllowProcessReuse);
+	}
 
-	//info.remoteLogEnabled = true;
-	UbaSessionServer = SessionServer_Create(info);
+
+	UbaServer = NetworkServer_Create(LogWriter);
+	UbaStorageServer = StorageServer_Create2(*UbaServer, *Config, LogWriter);
+	UbaSessionServer = SessionServer_Create2(*UbaStorageServer, *UbaServer, *Config, LogWriter);
+	UbaScheduler = Scheduler_Create2(*UbaSessionServer, *Config);
+
+	Config_Destroy(Config);
+
+	// Config used by clients that connect
+	{
+		uba::Config* ClientConfig = Config_Create();
+		uba::ConfigTable* StorageTable = Config_AddTable(*ClientConfig, TC("Storage"));
+		ConfigTable_AddValueBool(*StorageTable,  TC("ResendCas"), true); // Since we call StorageServer_DeleteFile there is a tiny risk we might delete a cas file that is needed in the future
+		NetworkServer_SetClientsConfig(UbaServer, *ClientConfig);
+		Config_Destroy(ClientConfig);
+	}
 
 	CalculateKnownInputs();
 	UpdateMaxLocalParallelJobs();
 
-	UbaScheduler = Scheduler_Create(UbaSessionServer, MaxLocalParallelJobs, UbaJobProcessorOptions::bAllowProcessReuse);
 	Scheduler_Start(UbaScheduler);
 
 	HandleTaskQueueUpdated(TEXT("")); // Flush tasks into uba scheduler
