[UBA] Add UbaHost.dll to RuntimeDependencies in UbaController build script.

This adds UbaHost.dll and its other platform specific files to RuntimeDependencies to ensure it's staged alongside the target project.
It is meant to resovle a runtime error causing a crash inside WinAPI's "_delayLoadHelper2()" due to a missing DLL.

#rnx

diff --git a/Engine/Plugins/UbaController/Source/UbaController.build.cs b/Engine/Plugins/UbaController/Source/UbaController.build.cs
--- a/Engine/Plugins/UbaController/Source/UbaController.build.cs
+++ b/Engine/Plugins/UbaController/Source/UbaController.build.cs
@@ -40,14 +40,17 @@ public class UbaController : ModuleRules
 			PrivateRuntimeLibraryPaths.Add(Path.Combine(ModuleDirectory, "lib"));
 
 			PublicDelayLoadDLLs.Add("UbaHost.dll");
+			RuntimeDependencies.Add(Path.Combine(UbaBinariesPath, "UbaHost.dll"));
 		}
 		else if (Target.Platform == UnrealTargetPlatform.Mac)
 		{
 			PublicDelayLoadDLLs.Add(Path.Combine(UbaBinariesPath, "libUbaHost.dylib"));
+			RuntimeDependencies.Add(Path.Combine(UbaBinariesPath, "libUbaHost.dylib"));
 		}
 		else if (Target.Platform == UnrealTargetPlatform.Linux)
 		{
 			PublicDelayLoadDLLs.Add(Path.Combine(UbaBinariesPath, "libUbaHost.so"));
+			RuntimeDependencies.Add(Path.Combine(UbaBinariesPath, "libUbaHost.so"));
 		}
 	}
 }
