[ubacontroller] Add simple heartbeat logging (default to once every 3 mins) to indicate status of tasks while executing (will only print while there are tasks to process, i.e. between the "Starting up UBA/Horde connection" and "Shutting down UBA/Horde connection" messages). The intent is to be able to see from logs of running jobs whether anything abnormal is occurring with UBA (currently we can only see this after the fact by analyzing traces). Also cleaned up the other interval tracking to use FPlatformTime::Seconds directly rather than manually converting from cycles.

diff --git a/Engine/Plugins/UbaController/Source/Private/UbaJobProcessor.cpp b/Engine/Plugins/UbaController/Source/Private/UbaJobProcessor.cpp
--- a/Engine/Plugins/UbaController/Source/Private/UbaJobProcessor.cpp
+++ b/Engine/Plugins/UbaController/Source/Private/UbaJobProcessor.cpp
@@ -36,6 +36,12 @@ namespace UbaJobProcessorOptions
         MaxTimeWithoutTasks,
         TEXT("Time to wait (in seconds) before stop processing attempts if we don't have any pending task.\n"));
 
+	static float HeartBeatInterval = 180.0f;
+	static FAutoConsoleVariableRef CVarHeartBeatInterval(
+		TEXT("r.UbaController.HeartBeatInterval"),
+		HeartBeatInterval,
+		TEXT("Time between heart beat log messages"));
+
 	static bool bAutoLaunchVisualizer = false;
 	static FAutoConsoleVariableRef CVarAutoLaunchVisualizer(
 		TEXT("r.UbaController.AutoLaunchVisualizer"),
@@ -334,6 +340,8 @@ FString GetUbaBinariesPath();
 
 void FUbaJobProcessor::StartUba()
 {
+	UE_LOG(LogUbaController, Display, TEXT("Starting up UBA/Horde connection"));
+
 	checkf(UbaServer == nullptr, TEXT("FUbaJobProcessor::StartUba() was called twice before FUbaJobProcessor::ShutDownUba()"));
 
 	FString RootDir = FString::Printf(TEXT("%s/%s/%u"), *FUbaControllerModule::GetTempDir(), TEXT("UbaControllerStorageDir"), UE::GetMultiprocessId());
@@ -441,11 +449,12 @@ uint32 FUbaJobProcessor::Run()
 {
 	bIsWorkDone = false;
 	
-	uint32 LastTimeSinceHadJobs = FPlatformTime::Cycles();	
+	double LastTimeSinceHadJobs = FPlatformTime::Seconds();
+	double LastHeartBeat = FPlatformTime::Seconds();
 
 	while (!bForceStop)
 	{
-		const float ElapsedSeconds = (FPlatformTime::Cycles() - LastTimeSinceHadJobs) * FPlatformTime::GetSecondsPerCycle();
+		const double ElapsedSeconds = (FPlatformTime::Seconds() - LastTimeSinceHadJobs);
 
 		uint32 NumQueuedJobs = 0;
 		uint32 NumActiveLocalJobs = 0;
@@ -454,6 +463,8 @@ uint32 FUbaJobProcessor::Run()
 
 		FScopeLock lock(&bShouldProcessJobsLock);
 
+		const double HeartBeatElapsedSeconds = (FPlatformTime::Seconds() - LastHeartBeat);
+
 		if (UbaScheduler)
 		{
 			Scheduler_GetStats(UbaScheduler, NumQueuedJobs, NumActiveLocalJobs, NumActiveRemoteJobs, NumFinishedJobs);
@@ -479,7 +490,7 @@ uint32 FUbaJobProcessor::Run()
 				bShouldProcessJobs = true;
 			}
 
-			LastTimeSinceHadJobs = FPlatformTime::Cycles();
+			LastTimeSinceHadJobs = FPlatformTime::Seconds();
 		}
 
 		if (bShouldProcessJobs)
@@ -500,6 +511,13 @@ uint32 FUbaJobProcessor::Run()
 			SessionServer_SetMaxRemoteProcessCount(UbaSessionServer, MaxRemoteCoresToRequest);
 
 			UpdateStats();
+
+			if (HeartBeatElapsedSeconds > UbaJobProcessorOptions::HeartBeatInterval)
+			{
+				// only print heartbeat log messages if currently executing tasks
+				UE_LOG(LogUbaController, Display, TEXT("Task Status -- Queued: %u -- Active: %u local, %u remote -- Completed: %u"), NumQueuedJobs, NumActiveLocalJobs, NumActiveRemoteJobs, NumFinishedJobs);
+				LastHeartBeat = FPlatformTime::Seconds();
+			}
 		}
 
 		lock.Unlock();
