[UbaController]
* Changed so controller use Scheduler_IsEmpty() instead of GetStats since GetStats is not threadsafe in the sense that values might be in between states so adding them up might give the wrong total

diff --git a/Engine/Plugins/UbaController/Source/Private/UbaJobProcessor.cpp b/Engine/Plugins/UbaController/Source/Private/UbaJobProcessor.cpp
--- a/Engine/Plugins/UbaController/Source/Private/UbaJobProcessor.cpp
+++ b/Engine/Plugins/UbaController/Source/Private/UbaJobProcessor.cpp
@@ -370,6 +370,7 @@ void FUbaJobProcessor::StartUba()
 		ConfigTable_AddValueString(*SessionTable,  TC("TraceOutputFile"), TCHAR_TO_UBASTRING(*TraceOutputFile));
 		ConfigTable_AddValueBool(*SessionTable,    TC("DetailedTrace"), UbaJobProcessorOptions::bDetailedTrace);
 		ConfigTable_AddValueString(*SessionTable,  TC("TraceName"), TCHAR_TO_UBASTRING(*TraceName));
+		// ConfigTable_AddValueBool(*SessionTable,    TC("RemoteTraceEnabled"), true); Enable this to have the remotes send back the uba trace to the host (ends up in log folder)
 
 		uba::ConfigTable* SchedulerTable = Config_AddTable(*Config, TC("Scheduler"));
 		ConfigTable_AddValueU32(*SchedulerTable,   TC("MaxLocalProcessors"), MaxLocalParallelJobs);
@@ -459,6 +460,7 @@ uint32 FUbaJobProcessor::Run()
 		uint32 NumActiveLocalJobs = 0;
 		uint32 NumActiveRemoteJobs = 0;
 		uint32 NumFinishedJobs = 0;
+		bool bIsEmpty = true;
 
 		FScopeLock lock(&bShouldProcessJobsLock);
 
@@ -467,11 +469,12 @@ uint32 FUbaJobProcessor::Run()
 		if (UbaScheduler)
 		{
 			Scheduler_GetStats(UbaScheduler, NumQueuedJobs, NumActiveLocalJobs, NumActiveRemoteJobs, NumFinishedJobs);
+			bIsEmpty = Scheduler_IsEmpty(UbaScheduler);
 		}
 		const uint32 NumActiveJobs = NumActiveLocalJobs + NumActiveRemoteJobs;
 
 		// We don't want to hog up Horde resources.
-		if (bShouldProcessJobs && ElapsedSeconds > UbaJobProcessorOptions::MaxTimeWithoutTasks && (NumQueuedJobs + NumActiveJobs) == 0)
+		if (bShouldProcessJobs && bIsEmpty && ElapsedSeconds > UbaJobProcessorOptions::MaxTimeWithoutTasks)
 		{
 			// If we're optimizing job starting, we only want to shutdown UBA if all the processes have terminated
 			bShouldProcessJobs = false;
@@ -479,7 +482,7 @@ uint32 FUbaJobProcessor::Run()
 		}
 
 		// Check if we have new tasks to process
-		if ((ControllerModule.HasTasksDispatchedOrPending() || (NumQueuedJobs + NumActiveJobs) != 0))
+		if (!bIsEmpty || ControllerModule.HasTasksDispatchedOrPending())
 		{
 			if (!bShouldProcessJobs)
 			{
