[UBA] Add command line options that match the terminology of other distributed shader controllers.

Adds "NoUBAShaderCompile" and "NoShaderWorker" to UBA controller alongside "NoUBAController".

#rnx

diff --git a/Engine/Plugins/UbaController/Source/Private/UbaControllerModule.cpp b/Engine/Plugins/UbaController/Source/Private/UbaControllerModule.cpp
--- a/Engine/Plugins/UbaController/Source/Private/UbaControllerModule.cpp
+++ b/Engine/Plugins/UbaController/Source/Private/UbaControllerModule.cpp
@@ -87,7 +87,9 @@ FUbaControllerModule::~FUbaControllerModule()
 
 static bool IsUbaControllerEnabled()
 {
-	if (FParse::Param(FCommandLine::Get(), TEXT("NoUbaController")))
+	if (FParse::Param(FCommandLine::Get(), TEXT("NoUbaController")) || // Compatibility to old terminology
+		FParse::Param(FCommandLine::Get(), TEXT("NoUbaShaderCompile")) || // Same terminology as other shader controllers
+		FParse::Param(FCommandLine::Get(), TEXT("NoShaderWorker")))
 	{
 		return false;
 	}
