[UbaController]
* Fixed bug where multiple unrealeditor instances ended up using the same casdb path

diff --git a/Engine/Plugins/UbaController/Source/Private/UbaJobProcessor.cpp b/Engine/Plugins/UbaController/Source/Private/UbaJobProcessor.cpp
--- a/Engine/Plugins/UbaController/Source/Private/UbaJobProcessor.cpp
+++ b/Engine/Plugins/UbaController/Source/Private/UbaJobProcessor.cpp
@@ -343,7 +343,15 @@ void FUbaJobProcessor::StartUba()
 
 	checkf(UbaServer == nullptr, TEXT("FUbaJobProcessor::StartUba() was called twice before FUbaJobProcessor::ShutDownUba()"));
 
-	FString RootDir = FString::Printf(TEXT("%s/%s/%u"), *FUbaControllerModule::GetTempDir(), TEXT("UbaControllerStorageDir"), UE::GetMultiprocessId());
+	FString RootDir;
+	int32 FolderIndex = 0;
+	while (true)
+	{
+		RootDir = FString::Printf(TEXT("%s/%s/%u"), *FUbaControllerModule::GetTempDir(), TEXT("UbaControllerStorageDir"), FolderIndex);
+		if (Uba_GetExclusiveAccess(TCHAR_TO_UBASTRING(*RootDir)))
+			break;
+		++FolderIndex;
+	}
 	IFileManager::Get().MakeDirectory(*RootDir, true);
 
 	FString TraceOutputFile;
@@ -401,7 +409,7 @@ void FUbaJobProcessor::StartUba()
 
 	HandleTaskQueueUpdated(TEXT("")); // Flush tasks into uba scheduler
 
-	if (UE::GetMultiprocessId() == 0)
+	if (FolderIndex == 0)
 	{
 		NetworkServer_StartListen(UbaServer, uba::DefaultPort, nullptr); // Start listen so any helper on the LAN can join in
 	}
