fix for UbaController on arm64ec

#rnx

diff --git a/Engine/Plugins/UbaController/Source/UbaController.build.cs b/Engine/Plugins/UbaController/Source/UbaController.build.cs
--- a/Engine/Plugins/UbaController/Source/UbaController.build.cs
+++ b/Engine/Plugins/UbaController/Source/UbaController.build.cs
@@ -21,7 +21,6 @@ public class UbaController : ModuleRules
 			"Json"
 		});
 
-		string Arch = TargetRules.Architecture.WindowsName == "arm64ec" ? "arm64" : TargetRules.Architecture.WindowsName;
 		string UbaSourcePath = Path.Combine(EngineDirectory, "Source", "Programs", "UnrealBuildAccelerator");
 
 		PublicIncludePaths.Add(UbaSourcePath);
@@ -34,7 +33,7 @@ public class UbaController : ModuleRules
 		
 		if (Target.Platform == UnrealTargetPlatform.Win64)
 		{
-			UbaBinariesPath = Path.Combine(UbaBinariesPath, Arch);
+			UbaBinariesPath = Path.Combine(UbaBinariesPath, TargetRules.Architecture.WindowsLibDir);
 			
 			PublicAdditionalLibraries.Add(Path.Combine(UbaBinariesPath, "UbaHost.lib"));
 			PrivateRuntimeLibraryPaths.Add(Path.Combine(ModuleDirectory, "lib"));
diff --git a/Engine/Source/Programs/UnrealBuildTool/Platform/Windows/UEBuildWindows.cs b/Engine/Source/Programs/UnrealBuildTool/Platform/Windows/UEBuildWindows.cs
--- a/Engine/Source/Programs/UnrealBuildTool/Platform/Windows/UEBuildWindows.cs
+++ b/Engine/Source/Programs/UnrealBuildTool/Platform/Windows/UEBuildWindows.cs
@@ -1221,7 +1221,6 @@ namespace UnrealBuildTool
 						{
 							"Reflex", // nvapi
 							"ACLPlugin", // __vectorcall calling convention is not currently supported
-							"UbaController", // LNK2001 (EC Symbol) errors
 						});
 						Target.GlobalDefinitions.Add("UE_EXTERNAL_PROFILING_ENABLED=0");
 					}
