[UBA] Make use of UBA trace snapshots by adding CVar "r.UbaController.SaveTraceSnapshotInterval" and always save snapshot when processing an output file failed.

#rnx

diff --git a/Engine/Plugins/UbaController/Source/Private/UbaJobProcessor.cpp b/Engine/Plugins/UbaController/Source/Private/UbaJobProcessor.cpp
--- a/Engine/Plugins/UbaController/Source/Private/UbaJobProcessor.cpp
+++ b/Engine/Plugins/UbaController/Source/Private/UbaJobProcessor.cpp
@@ -76,6 +76,13 @@ namespace UbaJobProcessorOptions
 		TEXT("1 - Also forward regular information about UBA sessions.\n")
 		TEXT("2 - Forward all UBA logs."));
 
+	static int32 SaveUbaTraceSnapshotInterval = 0;
+	static FAutoConsoleVariableRef CVarSaveUbaTraceSnapshotInterval(
+		TEXT("r.UbaController.SaveTraceSnapshotInterval"),
+		SaveUbaTraceSnapshotInterval,
+		TEXT("Specifies the interval (in seconds) in which a snapshot of the current state of the UBA trace will be saved to file.\n")
+		TEXT("A value of 0 disables the periodic snapshots and only saves the UBA trace at the end of each server session. By default 0.\n"));
+
 	static bool bProcessLogEnabled = false;
 	static FAutoConsoleVariableRef CVarProcessLogEnabled(
 		TEXT("r.UbaController.ProcessLogEnabled"),
@@ -354,11 +361,10 @@ void FUbaJobProcessor::StartUba()
 	}
 	IFileManager::Get().MakeDirectory(*RootDir, true);
 
-	FString TraceOutputFile;
 	if (!ControllerModule.GetDebugInfoPath().IsEmpty())
 	{
 		static uint32 UbaSessionCounter;
-		TraceOutputFile = ControllerModule.GetDebugInfoPath() / FString::Printf(TEXT("UbaController.MultiprocessId-%u.Session-%u.uba"), UE::GetMultiprocessId(), UbaSessionCounter++);
+		TraceOutputFilename = ControllerModule.GetDebugInfoPath() / FString::Printf(TEXT("UbaController.MultiprocessId-%u.Session-%u.uba"), UE::GetMultiprocessId(), UbaSessionCounter++);
 	}
 	FString TraceName = FString::Printf(TEXT("UbaController_%s"), *FGuid::NewGuid().ToString(EGuidFormats::Digits));
 
@@ -375,7 +381,7 @@ void FUbaJobProcessor::StartUba()
 		ConfigTable_AddValueBool(*SessionTable,    TC("AllowMemoryMaps"), false); // Skip using memory maps
 		ConfigTable_AddValueBool(*SessionTable,    TC("RemoteLogEnabled"), UbaJobProcessorOptions::bProcessLogEnabled);
 		ConfigTable_AddValueBool(*SessionTable,    TC("TraceEnabled"), true);
-		ConfigTable_AddValueString(*SessionTable,  TC("TraceOutputFile"), TCHAR_TO_UBASTRING(*TraceOutputFile));
+		ConfigTable_AddValueString(*SessionTable,  TC("TraceOutputFile"), TCHAR_TO_UBASTRING(*TraceOutputFilename));
 		ConfigTable_AddValueBool(*SessionTable,    TC("DetailedTrace"), UbaJobProcessorOptions::bDetailedTrace);
 		ConfigTable_AddValueString(*SessionTable,  TC("TraceName"), TCHAR_TO_UBASTRING(*TraceName));
 		// ConfigTable_AddValueBool(*SessionTable,    TC("RemoteTraceEnabled"), true); Enable this to have the remotes send back the uba trace to the host (ends up in log folder)
@@ -457,8 +463,10 @@ uint32 FUbaJobProcessor::Run()
 {
 	bIsWorkDone = false;
 	
-	double LastTimeSinceHadJobs = FPlatformTime::Seconds();
-	double LastHeartBeat = FPlatformTime::Seconds();
+	const double CurrentTime = FPlatformTime::Seconds();
+	double LastTimeSinceHadJobs = CurrentTime;
+	double LastHeartBeat = CurrentTime;
+	double LastTraceSnapshot = CurrentTime;
 
 	while (!bForceStop)
 	{
@@ -528,6 +536,17 @@ uint32 FUbaJobProcessor::Run()
 				UE_LOG(LogUbaController, Display, TEXT("Task Status -- Queued: %u -- Active: %u local, %u remote -- Completed: %u"), NumQueuedJobs, NumActiveLocalJobs, NumActiveRemoteJobs, NumFinishedJobs);
 				LastHeartBeat = FPlatformTime::Seconds();
 			}
+
+			// Save snapshot of current trace if periodic saves are enabled
+			if (UbaJobProcessorOptions::SaveUbaTraceSnapshotInterval > 0 && UbaSessionServer)
+			{
+				const int32 SnapshotIntervalElapsedSeconds = (int32)(FPlatformTime::Seconds() - LastTraceSnapshot);
+				if (SnapshotIntervalElapsedSeconds >= UbaJobProcessorOptions::SaveUbaTraceSnapshotInterval)
+				{
+					SaveSnapshotOfTrace();
+					LastTraceSnapshot = FPlatformTime::Seconds();
+				}
+			}
 		}
 
 		lock.Unlock();
@@ -650,6 +669,9 @@ void FUbaJobProcessor::HandleUbaJobFinished(FDistributedBuildTask* CompileTask)
 	const bool bWasSuccessful = ProcessOutputFile(CompileTask);
 	if (!bWasSuccessful)
 	{
+		// Save current snapshot of UBA trace in case this failure crashes the cook later on
+		SaveSnapshotOfTrace();
+
 		// If it failed running locally, lets try Run it locally but outside Uba
 		// Signaling a jobs as complete when it wasn't done, will cause a rerun on local worker as fallback
 		// because there is not output file for this job
@@ -715,3 +737,12 @@ void FUbaJobProcessor::UpdateStats()
 		Stats.MaxActiveAgentCores = FMath::Max(Stats.MaxActiveAgentCores, HordeAgentManager->GetActiveCoreCount());
 	}
 }
+
+void FUbaJobProcessor::SaveSnapshotOfTrace()
+{
+	if (!TraceOutputFilename.IsEmpty())
+	{
+		UE_LOG(LogUbaController, Log, TEXT("Save snapshot of UBA trace: %s"), *TraceOutputFilename);
+		SessionServer_SaveSnapshotOfTrace(UbaSessionServer);
+	}
+}
diff --git a/Engine/Plugins/UbaController/Source/Public/UbaJobProcessor.h b/Engine/Plugins/UbaController/Source/Public/UbaJobProcessor.h
--- a/Engine/Plugins/UbaController/Source/Public/UbaJobProcessor.h
+++ b/Engine/Plugins/UbaController/Source/Public/UbaJobProcessor.h
@@ -51,6 +51,8 @@ private:
 
 	void UpdateStats();
 
+	void SaveSnapshotOfTrace();
+
 	/** The runnable thread */
 	FRunnableThread* Thread = nullptr;
 	
@@ -71,6 +73,8 @@ private:
 	/** Set to true when the main loop finishes*/
 	bool bIsWorkDone = false;
 
+	FString TraceOutputFilename;
+
 	uba::NetworkServer* UbaServer = nullptr;
 	uba::StorageServer* UbaStorageServer = nullptr;
 	uba::SessionServer* UbaSessionServer = nullptr;
