[UBA]
* Added support for linux in UbaController

diff --git a/Engine/Plugins/UbaController/Source/Private/UbaControllerModule.cpp b/Engine/Plugins/UbaController/Source/Private/UbaControllerModule.cpp
--- a/Engine/Plugins/UbaController/Source/Private/UbaControllerModule.cpp
+++ b/Engine/Plugins/UbaController/Source/Private/UbaControllerModule.cpp
@@ -180,7 +180,11 @@ void FUbaControllerModule::LoadDependencies()
 {
 	const FString UbaBinariesPath = GetUbaBinariesPath();
 	FPlatformProcess::AddDllDirectory(*UbaBinariesPath);
+	#if PLATFORM_WINDOWS
 	FPlatformProcess::GetDllHandle(*(FPaths::Combine(UbaBinariesPath, "UbaHost.dll")));
+	#elif PLATFORM_LINUX
+	FPlatformProcess::GetDllHandle(*(FPaths::Combine(UbaBinariesPath, "libUbaHost.so")));
+	#endif
 }
 
 void FUbaControllerModule::StartupModule()
diff --git a/Engine/Plugins/UbaController/Source/UbaController.build.cs b/Engine/Plugins/UbaController/Source/UbaController.build.cs
--- a/Engine/Plugins/UbaController/Source/UbaController.build.cs
+++ b/Engine/Plugins/UbaController/Source/UbaController.build.cs
@@ -26,8 +26,7 @@ public class UbaController : ModuleRules
 		PublicIncludePaths.Add(UbaSourcePath);
 		PublicIncludePaths.Add(Path.Combine(UbaSourcePath, "Core", "Public"));
 		PublicIncludePaths.Add(Path.Combine(UbaSourcePath, "Common", "Public"));
-        PublicIncludePaths.Add(Path.Combine(UbaSourcePath, "Host", "Public"));
-        PublicIncludePathModuleNames.Add("BLAKE3");
+		PublicIncludePaths.Add(Path.Combine(UbaSourcePath, "Host", "Public"));
 
 		string UbaBinariesPath = Path.Combine(EngineDirectory, "Binaries", Target.Platform.ToString(), "UnrealBuildAccelerator");
 		
