[UBA]
* Enabled UbaController plugin for mac and linux

diff --git a/Engine/Plugins/UbaController/UbaController.uplugin b/Engine/Plugins/UbaController/UbaController.uplugin
--- a/Engine/Plugins/UbaController/UbaController.uplugin
+++ b/Engine/Plugins/UbaController/UbaController.uplugin
@@ -22,9 +22,9 @@
 			"Name" : "UbaController",
 			"Type" : "Editor",
 			"LoadingPhase" : "EarliestPossible",
-			"PlatformAllowList": ["Win64"]
+			"PlatformAllowList": ["Win64", "Mac", "Linux"]
 		}
 	],
 
-	"SupportedTargetPlatforms": ["Win64"]
+	"SupportedTargetPlatforms": ["Win64", "Mac", "Linux"]
 }
\ No newline at end of file
