[ubacontroller] improvements
- remove the unnecessary critical section & "handletaskqueueupdated" mechanism; we're using a lockless threadsafe queue for pending jobs, so should just use that
- change the mode of the lockless queue to single producer/single consumer; only the distributed compilation thread pushes things into it so mpsc is unnecessary
- eliminate some confusing code around determining whether there is work to do. this is a simple check: see if the pending queue is non-empty, or if there's work in progress on the scheduler. no need for the HasTasksDispatchedOrPending function, which was calling back into the thread that was the only place it was called...
- add session name (used for trace file) to logging when a session starts. this can be used to figure out which process is the one currently being visualized in UBAVisualizer when multiple are running at once.
- increase the max concurrent processes limit; this was set unnecessarily low (though we're not often exceeding the limit currently due to I/O bottlenecks which will be addressed in a forthcoming CL)

diff --git a/Engine/Plugins/UbaController/Source/Private/UbaControllerModule.cpp b/Engine/Plugins/UbaController/Source/Private/UbaControllerModule.cpp
--- a/Engine/Plugins/UbaController/Source/Private/UbaControllerModule.cpp
+++ b/Engine/Plugins/UbaController/Source/Private/UbaControllerModule.cpp
@@ -153,11 +153,6 @@ void FUbaControllerModule::CleanWorkingDirectory() const
 	}
 }
 
-bool FUbaControllerModule::HasTasksDispatchedOrPending() const
-{
-	return !PendingRequestedCompilationTasks.IsEmpty() || (JobDispatcherThread.IsValid() && JobDispatcherThread->HasJobsInFlight());
-}
-
 FString GetUbaBinariesPath()
 {
 #if PLATFORM_WINDOWS
@@ -298,8 +293,6 @@ TFuture<FDistributedBuildTaskResult> FUbaControllerModule::EnqueueTask(const FTa
 		PendingRequestedCompilationTasks.Enqueue(Task);
 	}
 
-	JobDispatcherThread->HandleTaskQueueUpdated(CommandData.InputFileName);
-
 	return MoveTemp(Future);
 }
 
diff --git a/Engine/Plugins/UbaController/Source/Private/UbaJobProcessor.cpp b/Engine/Plugins/UbaController/Source/Private/UbaJobProcessor.cpp
--- a/Engine/Plugins/UbaController/Source/Private/UbaJobProcessor.cpp
+++ b/Engine/Plugins/UbaController/Source/Private/UbaJobProcessor.cpp
@@ -134,7 +134,6 @@ FUbaJobProcessor::FUbaJobProcessor(
 	ControllerModule(InControllerModule),
 	bForceStop(false),
 	LastTimeCheckedForTasks(0),
-	bShouldProcessJobs(false),
 	bIsWorkDone(false),
 	LogWriter([]() {}, []() {}, [](uba::LogEntryType type, const uba::tchar* str, uba::u32 /*strlen*/)
 	{
@@ -346,7 +345,8 @@ FString GetUbaBinariesPath();
 
 void FUbaJobProcessor::StartUba()
 {
-	UE_LOG(LogUbaController, Display, TEXT("Starting up UBA/Horde connection"));
+	FString TraceName = FString::Printf(TEXT("UbaController_%s"), *FGuid::NewGuid().ToString(EGuidFormats::Digits));
+	UE_LOG(LogUbaController, Display, TEXT("Starting up UBA/Horde connection for session %s"), *TraceName);
 
 	checkf(UbaServer == nullptr, TEXT("FUbaJobProcessor::StartUba() was called twice before FUbaJobProcessor::ShutDownUba()"));
 
@@ -366,7 +366,6 @@ void FUbaJobProcessor::StartUba()
 		static uint32 UbaSessionCounter;
 		TraceOutputFilename = ControllerModule.GetDebugInfoPath() / FString::Printf(TEXT("UbaController.MultiprocessId-%u.Session-%u.uba"), UE::GetMultiprocessId(), UbaSessionCounter++);
 	}
-	FString TraceName = FString::Printf(TEXT("UbaController_%s"), *FGuid::NewGuid().ToString(EGuidFormats::Digits));
 
 	uba::Config* Config = Config_Create();
 	{
@@ -413,8 +412,6 @@ void FUbaJobProcessor::StartUba()
 
 	Scheduler_Start(UbaScheduler);
 
-	HandleTaskQueueUpdated(TEXT("")); // Flush tasks into uba scheduler
-
 	if (FolderIndex == 0)
 	{
 		NetworkServer_StartListen(UbaServer, uba::DefaultPort, nullptr); // Start listen so any helper on the LAN can join in
@@ -482,43 +479,53 @@ uint32 FUbaJobProcessor::Run()
 		uint32 NumActiveLocalJobs = 0;
 		uint32 NumActiveRemoteJobs = 0;
 		uint32 NumFinishedJobs = 0;
-		bool bIsEmpty = true;
-
-		FScopeLock lock(&bShouldProcessJobsLock);
+		bool bNewTasks = !ControllerModule.PendingRequestedCompilationTasks.IsEmpty();
+		// never empty if we have new tasks
+		bool bIsEmpty = !bNewTasks;
 
 		const double HeartBeatElapsedSeconds = (FPlatformTime::Seconds() - LastHeartBeat);
 
 		if (UbaScheduler)
 		{
 			Scheduler_GetStats(UbaScheduler, NumQueuedJobs, NumActiveLocalJobs, NumActiveRemoteJobs, NumFinishedJobs);
-			bIsEmpty = Scheduler_IsEmpty(UbaScheduler);
+			bIsEmpty &= Scheduler_IsEmpty(UbaScheduler);
 		}
 		const uint32 NumActiveJobs = NumActiveLocalJobs + NumActiveRemoteJobs;
 
 		// We don't want to hog up Horde resources.
-		if (bShouldProcessJobs && bIsEmpty && ElapsedSeconds > UbaJobProcessorOptions::MaxTimeWithoutTasks)
+		if (UbaScheduler && bIsEmpty && ElapsedSeconds > UbaJobProcessorOptions::MaxTimeWithoutTasks)
 		{
 			// If we're optimizing job starting, we only want to shutdown UBA if all the processes have terminated
-			bShouldProcessJobs = false;
 			ShutDownUba();
 		}
 
 		// Check if we have new tasks to process
-		if (!bIsEmpty || ControllerModule.HasTasksDispatchedOrPending())
+		if (!bIsEmpty)
 		{
-			if (!bShouldProcessJobs)
+			if (!UbaScheduler)
 			{
 				// We have new tasks. Start processing again
 				StartUba();
-
-				bShouldProcessJobs = true;
 			}
 
 			LastTimeSinceHadJobs = FPlatformTime::Seconds();
 		}
 
-		if (bShouldProcessJobs)
+		if (UbaScheduler)
 		{
+			if (bNewTasks)
+			{
+				while (true)
+				{
+					FDistributedBuildTask* Task = nullptr;
+					if (!ControllerModule.PendingRequestedCompilationTasks.Dequeue(Task) || !Task)
+					{
+						break;
+					}
+					RunTaskWithUba(Task);
+				}
+			}
+
 			const int32 MaxAvailableLocalCores = FPlatformMisc::NumberOfCoresIncludingHyperthreads();
 			const int32 LocalCoresToNotUse = 1 + (int32)NumActiveRemoteJobs / 30; // Use one core per 30 remote ones
 			const int32 MaxLocalCoresToUse = FMath::Clamp(MaxAvailableLocalCores - LocalCoresToNotUse, 0, MaxLocalParallelJobs);
@@ -555,13 +562,9 @@ uint32 FUbaJobProcessor::Run()
 			}
 		}
 
-		lock.Unlock();
-
 		FPlatformProcess::Sleep(UbaJobProcessorOptions::SleepTimeBetweenActions);
 	}
 
-	FScopeLock lock(&bShouldProcessJobsLock);
-
 	ShutDownUba();
 
 	bIsWorkDone = true;
@@ -689,38 +692,14 @@ void FUbaJobProcessor::HandleUbaJobFinished(FDistributedBuildTask* CompileTask)
 	}
 }
 
-void FUbaJobProcessor::HandleTaskQueueUpdated(const FString& InputFileName)
-{
-	FScopeLock lock(&bShouldProcessJobsLock);
-
-	if (!UbaScheduler)
-	{
-		return;
-	}
-
-	while (true)
-	{
-		FDistributedBuildTask* Task = nullptr;
-		if (!ControllerModule.PendingRequestedCompilationTasks.Dequeue(Task) || !Task)
-		{
-			break;
-		}
-		RunTaskWithUba(Task);
-	}
-}
-
 bool FUbaJobProcessor::HasJobsInFlight() const
 {
 	if (!UbaScheduler)
 	{
 		return false;
 	}
-	uint32 queued = 0;
-	uint32 activeLocal = 0;
-	uint32 activeRemote = 0;
-	uint32 finished = 0;
-	Scheduler_GetStats(UbaScheduler, queued, activeLocal, activeRemote, finished);	
-	return (queued + activeLocal + activeRemote) != 0;
+
+	return !Scheduler_IsEmpty(UbaScheduler);
 }
 
 bool FUbaJobProcessor::PollStats(FDistributedBuildStats& OutStats)
diff --git a/Engine/Plugins/UbaController/Source/Public/UbaControllerModule.h b/Engine/Plugins/UbaController/Source/Public/UbaControllerModule.h
--- a/Engine/Plugins/UbaController/Source/Public/UbaControllerModule.h
+++ b/Engine/Plugins/UbaController/Source/Public/UbaControllerModule.h
@@ -36,10 +36,8 @@ public:
 	const FString& GetWorkingDirectory() const { return WorkingDirectory; }
 	const FString& GetDebugInfoPath() const { return DebugInfoPath; }
 
-	bool HasTasksDispatchedOrPending() const;
-
 	// Queue of tasks submitted by the engine, but not yet dispatched to the controller.
-	TQueue<FDistributedBuildTask*, EQueueMode::Mpsc> PendingRequestedCompilationTasks;
+	TQueue<FDistributedBuildTask*, EQueueMode::Spsc> PendingRequestedCompilationTasks;
 
 	virtual void SetMaxLocalWorkers(int32 InMaxNumLocalWorkers) override final
 	{
diff --git a/Engine/Plugins/UbaController/Source/Public/UbaJobProcessor.h b/Engine/Plugins/UbaController/Source/Public/UbaJobProcessor.h
--- a/Engine/Plugins/UbaController/Source/Public/UbaJobProcessor.h
+++ b/Engine/Plugins/UbaController/Source/Public/UbaJobProcessor.h
@@ -34,8 +34,6 @@ public:
 	/** Used to know when this thread has finished the main loop */
 	bool IsWorkDone() const { return bIsWorkDone;};
 
-	void HandleTaskQueueUpdated(const FString& InputFileName);
-
 	bool HasJobsInFlight() const;
 
 	bool PollStats(FDistributedBuildStats& OutStats);
@@ -66,10 +64,6 @@ private:
 	/** Used to abort the current processing loop */
 	uint32 LastTimeCheckedForTasks = 0;
 	
-	/** Used to stop the processing loop in the next main loop */
-	FCriticalSection bShouldProcessJobsLock;
-	bool bShouldProcessJobs = false;
-
 	/** Set to true when the main loop finishes*/
 	bool bIsWorkDone = false;
 
diff --git a/Engine/Source/Developer/UbaCoordinatorHorde/Public/UbaHordeConfig.h b/Engine/Source/Developer/UbaCoordinatorHorde/Public/UbaHordeConfig.h
--- a/Engine/Source/Developer/UbaCoordinatorHorde/Public/UbaHordeConfig.h
+++ b/Engine/Source/Developer/UbaCoordinatorHorde/Public/UbaHordeConfig.h
@@ -69,7 +69,7 @@ public:
 	FString UbaSentryUrl;
 
 	/** Maxmium number of CPU cores allowed to be used by build session. */
-	int32 HordeMaxCores = 576;
+	int32 HordeMaxCores = 2048;
 
 	/** Maximum number of local CPU cores allowed to be used by build session. By default -1, which indicates to use as many CPU cores as the local machine can provide. */
 	int32 MaxParallelActions = -1;
