[UBA] Disable UbaController on Mac to unblock EngineTest CIS builds.

#rnx

diff --git a/Engine/Plugins/UbaController/Source/Private/UbaControllerModule.cpp b/Engine/Plugins/UbaController/Source/Private/UbaControllerModule.cpp
--- a/Engine/Plugins/UbaController/Source/Private/UbaControllerModule.cpp
+++ b/Engine/Plugins/UbaController/Source/Private/UbaControllerModule.cpp
@@ -87,6 +87,11 @@ FUbaControllerModule::~FUbaControllerModule()
 
 static bool IsUbaControllerEnabled()
 {
+#if PLATFORM_MAC
+	// Currently disabled for Mac due to shadermap hangs and UBA detour issues on Mac.
+	return false;
+#endif
+
 	if (FParse::Param(FCommandLine::Get(), TEXT("NoUbaController")) || // Compatibility to old terminology
 		FParse::Param(FCommandLine::Get(), TEXT("NoUbaShaderCompile")) || // Same terminology as other shader controllers
 		FParse::Param(FCommandLine::Get(), TEXT("NoShaderWorker")))
