[ubacontroller] move task object deletion out of reportjobprocessed function and into root of task finished callback delegate to fix SA warning about multiple deletion (false positive, but this is simpler to follow in any event)

diff --git a/Engine/Plugins/UbaController/Source/Private/UbaControllerModule.cpp b/Engine/Plugins/UbaController/Source/Private/UbaControllerModule.cpp
--- a/Engine/Plugins/UbaController/Source/Private/UbaControllerModule.cpp
+++ b/Engine/Plugins/UbaController/Source/Private/UbaControllerModule.cpp
@@ -294,7 +294,6 @@ void FUbaControllerModule::ReportJobProcessed(const FTaskResponse& InTaskRespons
 	if (CompileTask)
 	{
 		CompileTask->Finalize(InTaskResponse.ReturnCode);
-		delete CompileTask;
 	}
 }
 
diff --git a/Engine/Plugins/UbaController/Source/Private/UbaJobProcessor.cpp b/Engine/Plugins/UbaController/Source/Private/UbaJobProcessor.cpp
--- a/Engine/Plugins/UbaController/Source/Private/UbaJobProcessor.cpp
+++ b/Engine/Plugins/UbaController/Source/Private/UbaJobProcessor.cpp
@@ -312,6 +312,7 @@ void FUbaJobProcessor::RunTaskWithUba(FDistributedBuildTask* Task)
 				StorageServer_DeleteFile(Info->Processor->UbaStorageServer, TCHAR_TO_UBASTRING(*Info->InputFile));
 				StorageServer_DeleteFile(Info->Processor->UbaStorageServer, TCHAR_TO_UBASTRING(*Info->OutputFile));
 
+				delete Info->Task;
 				delete Info;
 			}
 		};
