[UBA] Clean up dispatcher thread shutdown code in UbaController.

#rnx

diff --git a/Engine/Plugins/UbaController/Source/Private/UbaControllerModule.cpp b/Engine/Plugins/UbaController/Source/Private/UbaControllerModule.cpp
--- a/Engine/Plugins/UbaController/Source/Private/UbaControllerModule.cpp
+++ b/Engine/Plugins/UbaController/Source/Private/UbaControllerModule.cpp
@@ -75,13 +75,7 @@ FString FUbaControllerModule::GetTempDir()
 
 FUbaControllerModule::~FUbaControllerModule()
 {	
-	if (JobDispatcherThread)
-	{
-		JobDispatcherThread->Stop();
-		// Wait until the thread is done
-		FPlatformProcess::ConditionalSleep([&](){ return JobDispatcherThread && JobDispatcherThread->IsWorkDone(); },0.1f);
-	}
-
+	StopAndWaitForDispatcherThread();
 	CleanWorkingDirectory();
 }
 
@@ -182,11 +176,9 @@ void FUbaControllerModule::StartupModule()
 
 	FCoreDelegates::OnEnginePreExit.AddLambda([&]()
 	{
-		if (bControllerInitialized && JobDispatcherThread)
+		if (bControllerInitialized)
 		{
-			JobDispatcherThread->Stop();
-			FPlatformProcess::ConditionalSleep([&]() { return JobDispatcherThread && JobDispatcherThread->IsWorkDone(); }, 0.1f);
-			JobDispatcherThread = nullptr;
+			StopAndWaitForDispatcherThread();
 		}
 	});
 }
@@ -200,12 +192,7 @@ void FUbaControllerModule::ShutdownModule()
 	if (bControllerInitialized)
 	{
 		// Stop the jobs thread
-		if (JobDispatcherThread)
-		{
-			JobDispatcherThread->Stop();	
-			// Wait until the thread is done
-			FPlatformProcess::ConditionalSleep([&](){ return JobDispatcherThread && JobDispatcherThread->IsWorkDone(); },0.1f);
-		}
+		StopAndWaitForDispatcherThread();
 
 		FDistributedBuildTask* Task;
 		while (PendingRequestedCompilationTasks.Dequeue(Task))
@@ -243,10 +230,11 @@ void FUbaControllerModule::InitializeController()
 			{
 				DebugInfoPath = UbaControllerModule::MakeAndGetDebugInfoPath();
 			}
+
             // Make sure http module is loaded in game thread before launch UBA client
 			FModuleManager::LoadModuleChecked<FHttpModule>("HTTP");
-			JobDispatcherThread = MakeShared<FUbaJobProcessor>(*this);
-			JobDispatcherThread->StartThread();
+
+			StartDispatcherThread();
 		}
 
 		bControllerInitialized = true;
@@ -302,4 +290,22 @@ UBACONTROLLER_API FUbaControllerModule& FUbaControllerModule::Get()
 	return FModuleManager::LoadModuleChecked<FUbaControllerModule>(TEXT("UbaController"));
 }
 
+void FUbaControllerModule::StartDispatcherThread()
+{
+	checkf(JobDispatcherThread == nullptr, TEXT("UBA job dispatcher thread was not previoulsy shut down before starting it again"));
+	JobDispatcherThread = MakeUnique<FUbaJobProcessor>(*this);
+	JobDispatcherThread->StartThread();
+}
+
+void FUbaControllerModule::StopAndWaitForDispatcherThread()
+{
+	if (JobDispatcherThread)
+	{
+		JobDispatcherThread->Stop();
+		// Wait until the thread is done
+		FPlatformProcess::ConditionalSleep([&]() { return !JobDispatcherThread || JobDispatcherThread->IsWorkDone(); }, 0.1f);
+		JobDispatcherThread = nullptr;
+	}
+}
+
 IMPLEMENT_MODULE(FUbaControllerModule, UbaController);
diff --git a/Engine/Plugins/UbaController/Source/Public/UbaControllerModule.h b/Engine/Plugins/UbaController/Source/Public/UbaControllerModule.h
--- a/Engine/Plugins/UbaController/Source/Public/UbaControllerModule.h
+++ b/Engine/Plugins/UbaController/Source/Public/UbaControllerModule.h
@@ -54,6 +54,9 @@ public:
 private:
 	void LoadDependencies();
 
+	void StartDispatcherThread();
+	void StopAndWaitForDispatcherThread();
+
 	bool bSupported;
 	bool bModuleInitialized;
 	bool bControllerInitialized;
@@ -66,5 +69,5 @@ private:
 	TAtomic<int32> NextTaskID;
 	TAtomic<int32> MaxNumLocalWorkers = -1;
 
-	TSharedPtr<FUbaJobProcessor> JobDispatcherThread;
+	TUniquePtr<FUbaJobProcessor> JobDispatcherThread;
 };
