[UBA] Explicitly load Sockets module in UbaController to ensure it's loading in the game thread.

Otherwise, it might be loaded in a different thread during the initialization of the UbaAgentHordeManager.

#rnx

diff --git a/Engine/Plugins/UbaController/Source/Private/UbaControllerModule.cpp b/Engine/Plugins/UbaController/Source/Private/UbaControllerModule.cpp
--- a/Engine/Plugins/UbaController/Source/Private/UbaControllerModule.cpp
+++ b/Engine/Plugins/UbaController/Source/Private/UbaControllerModule.cpp
@@ -231,8 +231,9 @@ void FUbaControllerModule::InitializeController()
 				DebugInfoPath = UbaControllerModule::MakeAndGetDebugInfoPath();
 			}
 
-            // Make sure http module is loaded in game thread before launch UBA client
+            // Make sure HTTP and Scokets modules are loaded in game thread before launching UBA client
 			FModuleManager::LoadModuleChecked<FHttpModule>("HTTP");
+			FModuleManager::LoadModuleChecked<FHttpModule>("Sockets");
 
 			StartDispatcherThread();
 		}
