[UbaController]
* Added detours lib and agent executable as additional properties for receipt

diff --git a/Engine/Plugins/UbaController/Source/UbaController.build.cs b/Engine/Plugins/UbaController/Source/UbaController.build.cs
--- a/Engine/Plugins/UbaController/Source/UbaController.build.cs
+++ b/Engine/Plugins/UbaController/Source/UbaController.build.cs
@@ -39,16 +39,29 @@ public class UbaController : ModuleRules
 
 			PublicDelayLoadDLLs.Add("UbaHost.dll");
 			RuntimeDependencies.Add(Path.Combine(UbaBinariesPath, "UbaHost.dll"));
+			AdditionalPropertiesForReceipt.Add("Uba", Path.Combine(UbaBinariesPath, "UbaDetours.dll"));
+			AdditionalPropertiesForReceipt.Add("Uba", Path.Combine(UbaBinariesPath, "UbaAgent.exe"));
+
+			// todo: this only exists for x64 atm. Is not needed
+			//AdditionalPropertiesForReceipt.Add("Uba", Path.Combine(UbaBinariesPath, "UbaWine.dll.so"));
+
+			// Only for windows atm
+			AdditionalPropertiesForReceipt.Add("Uba", Path.Combine(UbaBinariesPath, "UbaHost.toml"));
+			AdditionalPropertiesForReceipt.Add("Uba", Path.Combine(UbaBinariesPath, "UbaAgent.toml"));
 		}
 		else if (Target.Platform == UnrealTargetPlatform.Mac)
 		{
 			PublicDelayLoadDLLs.Add(Path.Combine(UbaBinariesPath, "libUbaHost.dylib"));
 			RuntimeDependencies.Add(Path.Combine(UbaBinariesPath, "libUbaHost.dylib"));
+			AdditionalPropertiesForReceipt.Add("Uba", Path.Combine(UbaBinariesPath, "libUbaDetours.dylib"));
+			AdditionalPropertiesForReceipt.Add("Uba", Path.Combine(UbaBinariesPath, "UbaAgent"));
 		}
 		else if (Target.Platform == UnrealTargetPlatform.Linux)
 		{
 			PublicDelayLoadDLLs.Add(Path.Combine(UbaBinariesPath, "libUbaHost.so"));
 			RuntimeDependencies.Add(Path.Combine(UbaBinariesPath, "libUbaHost.so"));
+			AdditionalPropertiesForReceipt.Add("Uba", Path.Combine(UbaBinariesPath, "libUbaDetours.so"));
+			AdditionalPropertiesForReceipt.Add("Uba", Path.Combine(UbaBinariesPath, "UbaAgent"));
 		}
 	}
 }
