[ubacontroller] virtual file support & other optimizations/improvements
- use UBA virtual file mechanism for all inputs and outputs instead of needing to write them to disk and read them back unnecessarily, when the r.UbaController.UseVirtualFiles cvar is set to true (default on windows, currently disabled on mac/linux due to missing uba shared memory implementation)
- skip serializing source to the SCW input files in this case, instead registering additional virtual files for each source object and fixing up the FShaderSource on the SCW side after deserializing the input (avoids most of the memory overhead of having to keep the input files around in memory for the duration of the UBA task)
- remove unnecessary attaching of dependent source files to distributed build task; these are no longer needed since preprocessing occurs in the cooker process (also fix up fastbuild and sndbs controllers which were still using this)
- remove unnecessary temp file creation and file move in SCW; this is not needed when not actually writing the files to disk (and probably was never needed for UBA in general since it operates with completion callbacks instead of disk polling)
- modify job processor thread to use an event for synchronization with the controller module instead of sleeping; the controller will trigger the event whenever new jobs are added and wake up the thread if it's currently blocked
- remove unnecessary serialization of debugdescription to SCW input files; these are large, was causing both memory bloat and slower task "file" creation
- bump threshold at which cooker throttles package processing due to pending job queue; this was unnecessarily low
- modify behaviour determining count of agents to request and how many local cores to use. always use the number of cores available to us (similar to how local workers would operate), base remote help request based on the total queue size (not just how many we've currently submitted), update the latter on an interval to prevent getting stuck processing a big queue and only modifying how much help we ask for after the queue is fully processed

#Synced-CL 50093269
#Synced-CL 50093269

diff --git a/Engine/Config/BaseEditor.ini b/Engine/Config/BaseEditor.ini
--- a/Engine/Config/BaseEditor.ini
+++ b/Engine/Config/BaseEditor.ini
@@ -413,7 +413,7 @@ MemoryMaxUsedVirtual=0
 ; Unlike MemoryMinFreePhysical, it does not interact with MemoryMaxUsed; GC will be triggered whenever the
 ; event fires.
 MemoryTriggerGCAtPressureLevel=None
-MaxConcurrentShaderJobs=8192
+MaxConcurrentShaderJobs=65536
 MaxDiffsToLog=5
 IgnoreHeaderDiffs=True
 
diff --git a/Engine/Plugins/FastBuildController/Source/Private/FastBuildUtilities.cpp b/Engine/Plugins/FastBuildController/Source/Private/FastBuildUtilities.cpp
--- a/Engine/Plugins/FastBuildController/Source/Private/FastBuildUtilities.cpp
+++ b/Engine/Plugins/FastBuildController/Source/Private/FastBuildUtilities.cpp
@@ -20,23 +20,6 @@ namespace FASTBuildControllerUtilitiesVariables
         TEXT("0: Do not send along debug information in FASTBuild. \n")
         TEXT("1: Send along debug information in FASTBuild."),
         ECVF_Default);
-
-	int32 SendAllPossibleShaderDependencies = 1;
-	FAutoConsoleVariableRef CVarFASTBuildSendAllPossibleShaderDependencies(
-        TEXT("r.FASTBuildController.SendAllPossibleShaderDependencies"),
-        SendAllPossibleShaderDependencies,
-        TEXT("Send all possible dependencies of the shaders to the remote machines.")
-        TEXT("0: Use dependencies array reported in the task structure.\n")
-        TEXT("1: Brute-force discover all possible dependencies. \n"),
-        ECVF_Default);
-}
-
-FString FastBuildUtilities::GetShaderFullPathOfShaderDependency(const FString& InVirtualPath)
-{
-	FString DependencyFilename = GetShaderSourceFilePath(InVirtualPath);
-	DependencyFilename = IFileManager::Get().ConvertToAbsolutePathForExternalAppForWrite(*DependencyFilename);
-	FPaths::NormalizeDirectoryName(DependencyFilename);
-	return DependencyFilename;
 }
 
 void FastBuildUtilities::WriteShaderWorkerDependenciesToFile(FArchive& ScriptFile)
@@ -151,84 +134,6 @@ void FastBuildUtilities::WritePlatformCompilerDependenciesToFile(FArchive& Scrip
 	}
 }
 
-void FastBuildUtilities::WriteDependenciesForShaderToScript(const TArray<FDistributedBuildTask*>& InCompilationTasks, FArchive& ScriptFile)
-{
-	TArray<FString> FullUniqueDependenciesArray;
-	if (FASTBuildControllerUtilitiesVariables::SendAllPossibleShaderDependencies)
-	{
-		// This is kinda a hack because we are sending all possible dependencies
-		FDependencyUniqueArrayEnumerator ShaderUsfDeps = FDependencyUniqueArrayEnumerator(FullUniqueDependenciesArray, nullptr, TEXT(".usf"));
-		FDependencyUniqueArrayEnumerator ShaderUshDeps = FDependencyUniqueArrayEnumerator(FullUniqueDependenciesArray, nullptr, TEXT(".ush"));
-		FDependencyUniqueArrayEnumerator ShaderHeaderDeps = FDependencyUniqueArrayEnumerator(FullUniqueDependenciesArray, nullptr, TEXT(".h"));
-#if PLATFORM_WINDOWS
-		FDependencyUniqueArrayEnumerator ShaderHlsliDeps = FDependencyUniqueArrayEnumerator(FullUniqueDependenciesArray, nullptr, TEXT(".hlsli"));
-#endif
-
-		const TMap<FString, FString> ShaderSourceDirectoryMappings = AllShaderSourceDirectoryMappings();
-		for (auto& ShaderDirectoryMapping : ShaderSourceDirectoryMappings)
-		{
-			IFileManager::Get().IterateDirectoryRecursively(*ShaderDirectoryMapping.Value, ShaderUsfDeps);
-			IFileManager::Get().IterateDirectoryRecursively(*ShaderDirectoryMapping.Value, ShaderUshDeps);
-			IFileManager::Get().IterateDirectoryRecursively(*ShaderDirectoryMapping.Value, ShaderHeaderDeps);
-#if PLATFORM_WINDOWS
-			IFileManager::Get().IterateDirectoryRecursively(*ShaderDirectoryMapping.Value, ShaderHlsliDeps);
-#endif
-		}
-	}
-	else
-	{
-		{
-			for (FDistributedBuildTask* CompilationTask : InCompilationTasks)
-			{
-				for (const FString& Dependency : CompilationTask->CommandData.Dependencies)
-				{
-					FullUniqueDependenciesArray.AddUnique(GetShaderFullPathOfShaderDependency(Dependency));
-				}
-			}
-		}
-	}
-
-#if PLATFORM_MAC
-	const FString MetalIntermediateDir = FPaths::EngineIntermediateDir() + TEXT("/Shaders/metal");
-	FDependencyEnumerator MetalCompilerDeps = FDependencyEnumerator(ScriptFile, nullptr, nullptr);
-	IFileManager::Get().IterateDirectoryRecursively(*MetalIntermediateDir, MetalCompilerDeps);
-#endif
-
-	for (const FString& ExtraDependency : FullUniqueDependenciesArray)
-	{
-		FString SourcePath = ExtraDependency;
-		FString DestinationPath;
-		if (!FPaths::IsUnderDirectory(SourcePath, FPaths::RootDir()))
-		{
-			DestinationPath = FFastBuildControllerModule::Get().RemapPath(SourcePath);
-
-			IPlatformFile& PlatformFile = FPlatformFileManager::Get().GetPlatformFile();
-			FDateTime SourceTimeStamp = PlatformFile.GetTimeStamp(*SourcePath);
-			FDateTime DestinationTimeStamp = PlatformFile.GetTimeStamp(*DestinationPath);
-
-			if (SourceTimeStamp != FDateTime::MinValue() && SourceTimeStamp != DestinationTimeStamp)
-			{
-				FString DestinationDirectory = FPaths::GetPath(DestinationPath);
-				if (!PlatformFile.DirectoryExists(*DestinationDirectory))
-				{
-					PlatformFile.CreateDirectoryTree(*DestinationDirectory);
-				}
-				PlatformFile.CopyFile(*DestinationPath, *SourcePath);
-				PlatformFile.SetTimeStamp(*DestinationPath, SourceTimeStamp);
-			}
-
-			DestinationPath = IFileManager::Get().ConvertToAbsolutePathForExternalAppForWrite(*DestinationPath);
-		}
-		else
-		{
-			DestinationPath = SourcePath;
-		}
-
-		const FString ExtraFile = TEXT("\t\t'") + DestinationPath + TEXT("'," LINE_TERMINATOR_ANSI);
-		ScriptFile.Serialize((void*)StringCast<ANSICHAR>(*ExtraFile, ExtraFile.Len()).Get(), sizeof(ANSICHAR) * ExtraFile.Len());
-	}
-}
-
 void FastBuildUtilities::FASTBuildWriteScriptFileHeader(const TArray<FDistributedBuildTask*>& InCompilationTasks, FArchive& ScriptFile, const FString& WorkerName)
 {	
 	static constexpr TCHAR HeaderTemplate[] =
@@ -255,8 +160,6 @@ void FastBuildUtilities::FASTBuildWriteScriptFileHeader(const TArray<FDistribute
 
 	WriteShaderWorkerDependenciesToFile(ScriptFile);
 
-	WriteDependenciesForShaderToScript(InCompilationTasks, ScriptFile);
-
 	const FString ExtraFilesFooter =
 		TEXT("\t}" LINE_TERMINATOR_ANSI)
 		TEXT("}" LINE_TERMINATOR_ANSI);
diff --git a/Engine/Plugins/FastBuildController/Source/Public/FastBuildUtilities.h b/Engine/Plugins/FastBuildController/Source/Public/FastBuildUtilities.h
--- a/Engine/Plugins/FastBuildController/Source/Public/FastBuildUtilities.h
+++ b/Engine/Plugins/FastBuildController/Source/Public/FastBuildUtilities.h
@@ -81,10 +81,8 @@ public:
 	static FString GetFastBuildCache();
 
 protected:
-	static FString GetShaderFullPathOfShaderDependency(const FString& InVirtualPath);
 	static void WriteShaderWorkerDependenciesToFile(FArchive& ScriptFile);
 	static FString ReplaceEnvironmentVariablesInPath(const FString& ExtraFilePartialPath);
 	static void WriteDependenciesToFileUsingWildcardPath(FArchive& ScriptFile, FString ParsedPath);
 	static void WritePlatformCompilerDependenciesToFile(FArchive& ScriptFile);
-	static void WriteDependenciesForShaderToScript(const TArray<FDistributedBuildTask*>& InCompilationTasks, FArchive& ScriptFile);
 };
diff --git a/Engine/Plugins/UbaController/Source/Private/UbaControllerModule.cpp b/Engine/Plugins/UbaController/Source/Private/UbaControllerModule.cpp
--- a/Engine/Plugins/UbaController/Source/Private/UbaControllerModule.cpp
+++ b/Engine/Plugins/UbaController/Source/Private/UbaControllerModule.cpp
@@ -202,6 +202,7 @@ void FUbaControllerModule::ShutdownModule()
 		}
 
 		PendingRequestedCompilationTasks.Empty();
+		PendingRequestedCompilationCount = 0;
 	}
 
 	CleanWorkingDirectory();
@@ -257,6 +258,11 @@ FString FUbaControllerModule::CreateUniqueFilePath()
 	return FPaths::Combine(WorkingDirectory, FString::FromInt(FolderID), FString::Printf(TEXT("%d.uba"), FileID));
 }
 
+bool FUbaControllerModule::SupportsVirtualFiles() const
+{
+	return JobDispatcherThread != nullptr && JobDispatcherThread->IsUsingVirtualFiles();
+}
+
 TFuture<FDistributedBuildTaskResult> FUbaControllerModule::EnqueueTask(const FTaskCommandData& CommandData)
 {
 	check(bSupported);
@@ -268,6 +274,8 @@ TFuture<FDistributedBuildTaskResult> FUbaControllerModule::EnqueueTask(const FTa
 	FDistributedBuildTask* Task = new FDistributedBuildTask(NextTaskID++, CommandData, MoveTemp(Promise));
 	{
 		PendingRequestedCompilationTasks.Enqueue(Task);
+		PendingRequestedCompilationCount++;
+		ProcessorWakeEvent->Trigger();
 	}
 
 	return MoveTemp(Future);
@@ -278,14 +286,6 @@ bool FUbaControllerModule::PollStats(FDistributedBuildStats& OutStats)
 	return JobDispatcherThread != nullptr && JobDispatcherThread->PollStats(OutStats);
 }
 
-void FUbaControllerModule::ReportJobProcessed(const FTaskResponse& InTaskResponse, FDistributedBuildTask* CompileTask)
-{
-	if (CompileTask)
-	{
-		CompileTask->Finalize(InTaskResponse.ReturnCode);
-	}
-}
-
 UBACONTROLLER_API FUbaControllerModule& FUbaControllerModule::Get()
 {
 	return FModuleManager::LoadModuleChecked<FUbaControllerModule>(TEXT("UbaController"));
@@ -303,6 +303,7 @@ void FUbaControllerModule::StopAndWaitForDispatcherThread()
 	if (JobDispatcherThread)
 	{
 		JobDispatcherThread->Stop();
+		ProcessorWakeEvent->Trigger();
 		// Wait until the thread is done
 		FPlatformProcess::ConditionalSleep([&]() { return !JobDispatcherThread || JobDispatcherThread->IsWorkDone(); }, 0.1f);
 		JobDispatcherThread = nullptr;
diff --git a/Engine/Plugins/UbaController/Source/Private/UbaJobProcessor.cpp b/Engine/Plugins/UbaController/Source/Private/UbaJobProcessor.cpp
--- a/Engine/Plugins/UbaController/Source/Private/UbaJobProcessor.cpp
+++ b/Engine/Plugins/UbaController/Source/Private/UbaJobProcessor.cpp
@@ -6,6 +6,9 @@
 #include "Interfaces/ITargetPlatformManagerModule.h"
 #include "Misc/ConfigCacheIni.h"
 #include "Misc/CoreMisc.h"
+#include "Misc/FileHelper.h"
+#include "Misc/PathViews.h"
+#include "Serialization/MemoryReader.h"
 #include "UbaControllerModule.h"
 #include "UbaHordeAgentManager.h"
 #include "UbaHordeConfig.h"
@@ -24,17 +27,11 @@
 
 namespace UbaJobProcessorOptions
 {
-	static float SleepTimeBetweenActions = 0.01f;
-	static FAutoConsoleVariableRef CVarSleepTimeBetweenActions(
-        TEXT("r.UbaController.SleepTimeBetweenActions"),
-        SleepTimeBetweenActions,
-        TEXT("How much time the job processor thread should sleep between actions .\n"));
-
 	static float MaxTimeWithoutTasks = 100.0f;
 	static FAutoConsoleVariableRef CVarMaxTimeWithoutTasks(
-        TEXT("r.UbaController.MaxTimeWithoutTasks"),
-        MaxTimeWithoutTasks,
-        TEXT("Time to wait (in seconds) before stop processing attempts if we don't have any pending task.\n"));
+		TEXT("r.UbaController.MaxTimeWithoutTasks"),
+		MaxTimeWithoutTasks,
+		TEXT("Time to wait for pending tasks (in seconds) when the queue is empty before shutting down the UBA/Horde session.\n"));
 
 	static float HeartBeatInterval = 180.0f;
 	static FAutoConsoleVariableRef CVarHeartBeatInterval(
@@ -46,7 +43,7 @@ namespace UbaJobProcessorOptions
 	static FAutoConsoleVariableRef CVarProcessIdleTime(
 		TEXT("r.UbaController.ProcessIdleTime"),
 		ProcessIdleTime,
-		TEXT("ShaderCompileWorker idle time until it exits. Only used when process reuse is enabled"));
+		TEXT("Time ShaderCompileWorker will wait for new tasks before exiting. Only used when process reuse is enabled. Set to 0 to only re-use processes if there are tasks in the queue pending assignment."));
 
 	static bool bAutoLaunchVisualizer = false;
 	static FAutoConsoleVariableRef CVarAutoLaunchVisualizer(
@@ -95,6 +92,17 @@ namespace UbaJobProcessorOptions
 		bProcessLogEnabled,
 		TEXT("If true, each detoured process will write a log file. Note this is only useful if UBA is compiled in debug\n"));
 
+#if PLATFORM_WINDOWS
+	static bool bUseVirtualFiles = true;
+#else
+	// virtual file mechanism (which relies on shared memory) is not currently implemented in UBA for Mac/Linux, so we cannot use this on those platforms until that is fixed
+	static bool bUseVirtualFiles = false;
+#endif
+	static FAutoConsoleVariableRef CVarUseVirtualFiles(
+		TEXT("r.UbaController.UseVirtualFiles"),
+		bUseVirtualFiles,
+		TEXT("If true, UBA virtual file mechanism will be used for input and output files (avoiding writing them to disk)"));
+
 	FString ReplaceEnvironmentVariablesInPath(const FString& ExtraFilePartialPath) // Duplicated code with FAST build.. put it somewhere else?
 	{
 		FString ParsedPath;
@@ -276,33 +284,78 @@ void FUbaJobProcessor::CalculateKnownInputs()
 void FUbaJobProcessor::RunTaskWithUba(FDistributedBuildTask* Task)
 {
 	FTaskCommandData& Data = Task->CommandData;
-	SessionServer_RegisterNewFile(UbaSessionServer, TCHAR_TO_UBASTRING(*Data.InputFileName));
+
+	FString InputFileName = FPaths::GetCleanFilename(Data.InputFileName);
+	FString OutputFileName = FPaths::GetCleanFilename(Data.OutputFileName);
+	FString AppDir = FPaths::GetPath(Data.Command);
+	FStringToUbaStringConversion UbaInputFileNameStr(*InputFileName);
+	float TimeToLive = UbaJobProcessorOptions::bAllowProcessReuse ? UbaJobProcessorOptions::ProcessIdleTime : 0.0f;
+	FString Parameters = FString::Printf(TEXT("\"%s/\" %d 0 \"%s\" \"%s\" %s -TimeToLive=%.1f"), *Data.WorkingDirectory, Data.DispatcherPID, *InputFileName, *OutputFileName, *Data.ExtraCommandArgs, TimeToLive);
+
+	uba::Config* Config = Config_Create();
+	uba::ConfigTable* RootTable = Config_RootTable(*Config);
+	ConfigTable_AddValueString(*RootTable, TC("Application"), FStringToUbaStringConversion(*Data.Command).Get());
+	ConfigTable_AddValueString(*RootTable, TC("Arguments"), FStringToUbaStringConversion(*Parameters).Get());
+	ConfigTable_AddValueString(*RootTable, TC("Description"), UbaInputFileNameStr.Get());
+	ConfigTable_AddValueString(*RootTable, TC("WorkingDir"), FStringToUbaStringConversion(*AppDir).Get());
+	ConfigTable_AddValueString(*RootTable, TC("Breadcrumbs"), FStringToUbaStringConversion(*Data.Description).Get());
+	ConfigTable_AddValueBool(*RootTable,   TC("WriteOutputFilesOnFail"), true);
+
+	if (UbaJobProcessorOptions::bProcessLogEnabled)
+	{
+		ConfigTable_AddValueString(*RootTable, TC("LogFile"), UbaInputFileNameStr.Get());
+	}
+
+	uint32 SourceVirtualFileCount = 0;
+	if (UbaJobProcessorOptions::bUseVirtualFiles)
+	{
+		// create a map of "overlay" files for the process; doing so prevents dir/file table bloat by ensuring these temp files don't get stored in these tables indefinitely, only for the lifetime of the process
+		// note that this is only done if we're using virtual files since it requires shared memory to be enabled (and we keep this off when using the file-based path)
+		TStringBuilder<512> OverlayFiles;
+		OverlayFiles << Data.InputFileName;
+
+		// register the main input file as a virtual file
+		SessionServer_CreateVirtualFile(UbaSessionServer, TCHAR_TO_UBASTRING(*Data.InputFileName), Data.InputFileData.GetData(), (uint64)Data.InputFileData.NumBytes(), true, false);
+
+		// register each source file required for the batch of tasks as an individual virtual file, using the same base filename as the input file with a sourcefile index appended
+		TStringBuilder<512> SourceVirtualFileName;
+		for (const FMemoryView& Source : Data.SourceData)
+		{
+			SourceVirtualFileName.Reset();
+			SourceVirtualFileName << Data.InputFileName << SourceVirtualFileCount++;
+
+			OverlayFiles << ';' << SourceVirtualFileName;
+
+			SessionServer_CreateVirtualFile(UbaSessionServer, TCHAR_TO_UBASTRING(SourceVirtualFileName.ToString()), Source.GetData(), Source.GetSize(), true, false);
+		}
+
+		ConfigTable_AddValueString(*RootTable, TC("OverlayFiles"), FStringToUbaStringConversion(*OverlayFiles).Get());
+	}
+	else
+	{
+		SessionServer_RegisterNewFile(UbaSessionServer, TCHAR_TO_UBASTRING(*Data.InputFileName));
+	}
 
 	for (const FString& AdditionalOutputFolder : Data.AdditionalOutputFolders)
 	{
 		SessionServer_RegisterNewDirectory(UbaSessionServer, TCHAR_TO_UBASTRING(*AdditionalOutputFolder));
 	}
 
-	float TimeToLive = UbaJobProcessorOptions::bAllowProcessReuse ? UbaJobProcessorOptions::ProcessIdleTime : 0.0f;
-
-	FString InputFileName = FPaths::GetCleanFilename(Data.InputFileName);
-	FString OutputFileName = FPaths::GetCleanFilename(Data.OutputFileName);
-	FString Parameters = FString::Printf(TEXT("\"%s/\" %d 0 \"%s\" \"%s\" %s -TimeToLive=%.1f"), *Data.WorkingDirectory, Data.DispatcherPID, *InputFileName, *OutputFileName, *Data.ExtraCommandArgs, TimeToLive);
-	FString AppDir = FPaths::GetPath(Data.Command);
-	
-	struct ExitedInfo
+	struct FUbaProcessData
 	{
 		FUbaJobProcessor* Processor;
 		FString InputFile;
 		FString OutputFile;
 		FDistributedBuildTask* Task;
+		uint32 SourceVirtualFileCount;
 	};
 
-	auto Info = new ExitedInfo;
-	Info->Processor = this;
-	Info->InputFile = Data.InputFileName;
-	Info->OutputFile = Data.OutputFileName;
-	Info->Task = Task;
+	auto ProcessData = new FUbaProcessData;
+	ProcessData->Processor = this;
+	ProcessData->InputFile = Data.InputFileName;
+	ProcessData->OutputFile = Data.OutputFileName;
+	ProcessData->SourceVirtualFileCount = SourceVirtualFileCount;
+	ProcessData->Task = Task;
 
 	static auto ExitedFunc = [](void* UserData, const uba::ProcessHandle& Process)
 		{
@@ -311,38 +364,73 @@ void FUbaJobProcessor::RunTaskWithUba(FDistributedBuildTask* Task)
 				UE_LOG(LogUbaController, Display, TEXT("%s"), UBASTRING_TO_TCHAR(LogLine));
 			}
 
-			if (ExitedInfo* Info = static_cast<ExitedInfo*>(UserData)) // It can be null if custom message has already handled all of them
+			if (FUbaProcessData* ProcessData = static_cast<FUbaProcessData*>(UserData)) // It can be null if custom message has already handled all of them
 			{
-				IFileManager::Get().Delete(*Info->InputFile);
-				SessionServer_RegisterDeleteFile(Info->Processor->UbaSessionServer, TCHAR_TO_UBASTRING(*Info->InputFile));
-				Info->Processor->HandleUbaJobFinished(Info->Task);
+				struct FProcessOutputData
+				{
+					FUbaProcessData* ProcessData;
+					TArray<uint8> OutputFileBytes;
+				};
+				FProcessOutputData OutputData{ ProcessData };
 
-				StorageServer_DeleteFile(Info->Processor->UbaStorageServer, TCHAR_TO_UBASTRING(*Info->InputFile));
-				StorageServer_DeleteFile(Info->Processor->UbaStorageServer, TCHAR_TO_UBASTRING(*Info->OutputFile));
+				if (UbaJobProcessorOptions::bUseVirtualFiles)
+				{
+					SessionServer_DeleteVirtualFile(ProcessData->Processor->UbaSessionServer, TCHAR_TO_UBASTRING(*ProcessData->InputFile), false);
 
-				delete Info->Task;
-				delete Info;
-			}
-		};
+					TStringBuilder<512> SourceVirtualFileName;
+					for (uint32 SourceIndex = 0; SourceIndex < ProcessData->SourceVirtualFileCount; ++SourceIndex)
+					{
+						SourceVirtualFileName.Reset();
+						SourceVirtualFileName << ProcessData->InputFile << SourceIndex;
 
-	FStringToUbaStringConversion UbaInputFileNameStr(*InputFileName);
+						SessionServer_DeleteVirtualFile(ProcessData->Processor->UbaSessionServer, TCHAR_TO_UBASTRING(SourceVirtualFileName.ToString()), false);
+					}
 
-	uba::Config* Config = Config_Create();
-	uba::ConfigTable* RootTable = Config_RootTable(*Config);
-	ConfigTable_AddValueString(*RootTable, TC("Application"), FStringToUbaStringConversion(*Data.Command).Get());
-	ConfigTable_AddValueString(*RootTable, TC("Arguments"), FStringToUbaStringConversion(*Parameters).Get());
-	ConfigTable_AddValueString(*RootTable, TC("Description"), UbaInputFileNameStr.Get());
-	ConfigTable_AddValueString(*RootTable, TC("WorkingDir"), FStringToUbaStringConversion(*AppDir).Get());
-	ConfigTable_AddValueString(*RootTable, TC("Breadcrumbs"), FStringToUbaStringConversion(*Data.Description).Get());
-	ConfigTable_AddValueBool(*RootTable,   TC("WriteOutputFilesOnFail"), true);
+					// since we set ShouldWriteToDisk to false in the configuration, we need to manually process the output files from the job batch
+					// the main output file never needs to be written to disk; we process the data directly in-memory. all other outputs are debug
+					// artifacts, so those should be written out to disk if they exist in the process data
+					auto OutputFileCallback = [](void* InUserData, const uba::tchar* OutputFileName)
+						{
+							FProcessOutputData* UserData = static_cast<FProcessOutputData*>(InUserData);
+							FUbaProcessData* ProcessData = UserData->ProcessData;
+							FString OutputFileNameStr = UBASTRING_TO_TCHAR(OutputFileName);
+							FStringView BaseOutputFileName = FPathViews::GetCleanFilename(ProcessData->OutputFile);
+							if (OutputFileNameStr.EndsWith(BaseOutputFileName))
+							{
+								uint64 OutputFileSize = 0;
+								SessionServer_GetOutputFileSize(ProcessData->Processor->UbaSessionServer, OutputFileSize, OutputFileName);
+								checkf(OutputFileSize < INT_MAX, TEXT("SCW output file size too large for a 32-bit indexed TArray (%llu)"), OutputFileSize);
+								UserData->OutputFileBytes.SetNumUninitialized(OutputFileSize);
+
+								SessionServer_GetOutputFileData(ProcessData->Processor->UbaSessionServer, UserData->OutputFileBytes.GetData(), OutputFileName, true);
+							}
+							else
+							{
+								SessionServer_WriteOutputFile(ProcessData->Processor->UbaSessionServer, OutputFileName, true);
+							}
+						};
+					ProcessHandle_TraverseOutputFiles(const_cast<uba::ProcessHandle*>(&Process), OutputFileCallback, &OutputData);
+				}
+				else
+				{
+					IFileManager::Get().Delete(*ProcessData->InputFile);
+					FStringToUbaStringConversion InputFileUba(*ProcessData->InputFile);
+					SessionServer_RegisterDeleteFile(ProcessData->Processor->UbaSessionServer, InputFileUba.Get());
+					StorageServer_DeleteFile(ProcessData->Processor->UbaStorageServer, InputFileUba.Get());
+					
+					FFileHelper::LoadFileToArray(OutputData.OutputFileBytes, *ProcessData->OutputFile);
+					StorageServer_DeleteFile(ProcessData->Processor->UbaStorageServer, TCHAR_TO_UBASTRING(*ProcessData->OutputFile));
+
+				}
+				ProcessData->Processor->HandleUbaJobFinished(ProcessData->Task, MoveTemp(OutputData.OutputFileBytes), ProcessHandle_GetExitCode(&Process));
+				delete ProcessData->Task;
+				delete ProcessData;
+			}
+		};
 
-	if (UbaJobProcessorOptions::bProcessLogEnabled)
-	{
-		ConfigTable_AddValueString(*RootTable, TC("LogFile"), UbaInputFileNameStr.Get());
-	}
 
 	uba::ProcessStartInfo* StartInfo = ProcessStartInfo_Create3(*Config);
-	ProcessStartInfo_SetExitedCallback(*StartInfo, ExitedFunc, Info);
+	ProcessStartInfo_SetExitedCallback(*StartInfo, ExitedFunc, ProcessData);
 
 	Scheduler_EnqueueProcess(UbaScheduler, *StartInfo, 1.0f, KnownInputsBuffer.GetData(), KnownInputsBuffer.Num()*sizeof(uba::tchar), KnownInputsCount);
 
@@ -388,12 +476,14 @@ void FUbaJobProcessor::StartUba()
 
 		uba::ConfigTable* SessionTable = Config_AddTable(*Config, TC("Session"));
 		ConfigTable_AddValueBool(*SessionTable,    TC("LaunchVisualizer"), UbaJobProcessorOptions::bAutoLaunchVisualizer);
-		ConfigTable_AddValueBool(*SessionTable,    TC("AllowMemoryMaps"), false); // Skip using memory maps
+		ConfigTable_AddValueBool(*SessionTable,    TC("AllowMemoryMaps"), UbaJobProcessorOptions::bUseVirtualFiles);
 		ConfigTable_AddValueBool(*SessionTable,    TC("RemoteLogEnabled"), UbaJobProcessorOptions::bProcessLogEnabled);
 		ConfigTable_AddValueBool(*SessionTable,    TC("TraceEnabled"), true);
 		ConfigTable_AddValueString(*SessionTable,  TC("TraceOutputFile"), TCHAR_TO_UBASTRING(*TraceOutputFilename));
 		ConfigTable_AddValueBool(*SessionTable,    TC("DetailedTrace"), UbaJobProcessorOptions::bDetailedTrace);
 		ConfigTable_AddValueString(*SessionTable,  TC("TraceName"), TCHAR_TO_UBASTRING(*TraceName));
+		ConfigTable_AddValueBool(*SessionTable,    TC("ShouldWriteToDisk"), !UbaJobProcessorOptions::bUseVirtualFiles); // files generated by processes should not be written to disk if virtual files are enabled, this will be handled manually
+		
 		// ConfigTable_AddValueBool(*SessionTable,    TC("RemoteTraceEnabled"), true); Enable this to have the remotes send back the uba trace to the host (ends up in log folder)
 
 		uba::ConfigTable* SchedulerTable = Config_AddTable(*Config, TC("Scheduler"));
@@ -486,7 +576,7 @@ uint32 FUbaJobProcessor::Run()
 	{
 		const double ElapsedSeconds = (FPlatformTime::Seconds() - LastTimeSinceHadJobs);
 
-		uint32 NumQueuedJobs = 0;
+		uint32 NumUbaQueuedJobs = 0;
 		uint32 NumActiveLocalJobs = 0;
 		uint32 NumActiveRemoteJobs = 0;
 		uint32 NumFinishedJobs = 0;
@@ -498,7 +588,7 @@ uint32 FUbaJobProcessor::Run()
 
 		if (UbaScheduler)
 		{
-			Scheduler_GetStats(UbaScheduler, NumQueuedJobs, NumActiveLocalJobs, NumActiveRemoteJobs, NumFinishedJobs);
+			Scheduler_GetStats(UbaScheduler, NumUbaQueuedJobs, NumActiveLocalJobs, NumActiveRemoteJobs, NumFinishedJobs);
 			bIsEmpty &= Scheduler_IsEmpty(UbaScheduler);
 		}
 		const uint32 NumActiveJobs = NumActiveLocalJobs + NumActiveRemoteJobs;
@@ -524,40 +614,46 @@ uint32 FUbaJobProcessor::Run()
 
 		if (UbaScheduler)
 		{
+			Scheduler_SetMaxLocalProcessors(UbaScheduler, MaxLocalParallelJobs);
+
+			const int32 MaxRemoteCoresToRequest = FMath::Max(0, (int32)(NumUbaQueuedJobs + NumActiveJobs + ControllerModule.PendingRequestedCompilationCount) - MaxLocalParallelJobs);
+			if (HordeAgentManager)
+			{
+				HordeAgentManager->SetTargetCoreCount(MaxRemoteCoresToRequest);
+			}
+			
+			SessionServer_SetMaxRemoteProcessCount(UbaSessionServer, MaxRemoteCoresToRequest);
+
+
+			const float TaskProcessTimeSliceSeconds = 3.0f;
+			float TaskProcessTimeStart = FPlatformTime::Seconds();
+
 			if (bNewTasks)
 			{
 				while (true)
 				{
 					FDistributedBuildTask* Task = nullptr;
-					if (!ControllerModule.PendingRequestedCompilationTasks.Dequeue(Task) || !Task)
+					// break out of the task loop if we've been processing tasks for more than TaskProcessTimeSliceSeconds, or the pending tasks queue is empty
+					// the former ensures we don't go too long without updating the amount of remote help we request (this can cause pathological behaviour when
+					// the queue grows large quickly after being empty or small for a period of time).
+					if ((FPlatformTime::Seconds() > (TaskProcessTimeStart + TaskProcessTimeSliceSeconds))
+						|| !ControllerModule.PendingRequestedCompilationTasks.Dequeue(Task) 
+						|| !Task)
 					{
 						break;
 					}
+					ControllerModule.PendingRequestedCompilationCount--;
 					RunTaskWithUba(Task);
 				}
 			}
 
-			const int32 MaxAvailableLocalCores = FPlatformMisc::NumberOfCoresIncludingHyperthreads();
-			const int32 LocalCoresToNotUse = 1 + (int32)NumActiveRemoteJobs / 30; // Use one core per 30 remote ones
-			const int32 MaxLocalCoresToUse = FMath::Clamp(MaxAvailableLocalCores - LocalCoresToNotUse, 0, MaxLocalParallelJobs);
-			const int32 MaxRemoteCoresToRequest = FMath::Max(0, (int32)(NumQueuedJobs + NumActiveJobs) - MaxLocalCoresToUse);
-
-			Scheduler_SetMaxLocalProcessors(UbaScheduler, MaxLocalCoresToUse);
-
-			if (HordeAgentManager)
-			{
-				HordeAgentManager->SetTargetCoreCount(MaxRemoteCoresToRequest);
-			}
-			
-			// TODO: Not sure this is a good idea in a cooking scenario where number of queued processes are going up and down
-			SessionServer_SetMaxRemoteProcessCount(UbaSessionServer, MaxRemoteCoresToRequest);
 
 			UpdateStats();
 
 			if (HeartBeatElapsedSeconds > UbaJobProcessorOptions::HeartBeatInterval)
 			{
 				// only print heartbeat log messages if currently executing tasks
-				UE_LOG(LogUbaController, Display, TEXT("Task Status -- Queued: %u -- Active: %u local, %u remote -- Completed: %u"), NumQueuedJobs, NumActiveLocalJobs, NumActiveRemoteJobs, NumFinishedJobs);
+				UE_LOG(LogUbaController, Display, TEXT("Task Status -- Queued: %u submitted, %u pending -- Active: %u local, %u remote -- Completed: %u"), NumUbaQueuedJobs, ControllerModule.PendingRequestedCompilationCount.load(), NumActiveLocalJobs, NumActiveRemoteJobs, NumFinishedJobs);
 				LastHeartBeat = FPlatformTime::Seconds();
 			}
 
@@ -572,8 +668,14 @@ uint32 FUbaJobProcessor::Run()
 				}
 			}
 		}
-
-		FPlatformProcess::Sleep(UbaJobProcessorOptions::SleepTimeBetweenActions);
+		
+		if (ControllerModule.PendingRequestedCompilationTasks.IsEmpty())
+		{
+			// wait for new tasks if there are none currently available; we only wait on the wake event as long as our shutdown timeout
+			// this will be triggered when we have new tasks to process or need to terminate the session (either due to a stop signal 
+			// from the controller module, or due to not receiving any new work within the timeout)
+			ControllerModule.ProcessorWakeEvent->Wait(UbaJobProcessorOptions::MaxTimeWithoutTasks * 1000);
+		}
 	}
 
 	ShutDownUba();
@@ -587,120 +689,58 @@ void FUbaJobProcessor::Stop()
 	bForceStop = true;
 };
 
+bool FUbaJobProcessor::IsUsingVirtualFiles() const
+{
+	return UbaJobProcessorOptions::bUseVirtualFiles;
+}
+
 void FUbaJobProcessor::StartThread()
 {
 	Thread = FRunnableThread::Create(this, TEXT("UbaJobProcessor"), 0, TPri_SlightlyBelowNormal, FPlatformAffinity::GetPoolThreadMask());
 }
 
-bool FUbaJobProcessor::ProcessOutputFile(FDistributedBuildTask* CompileTask)
+void FUbaJobProcessor::HandleUbaJobFinished(FDistributedBuildTask* CompileTask, TArray<uint8>&& OutputFileData, uint32 ExitCode)
 {
-	//TODO: This method is mostly taken from the other Distribution controllers
-	// As we get an explicit callback when the process ends, we should be able to simplify this to just check if the file exists
-	IPlatformFile& PlatformFile = FPlatformFileManager::Get().GetPlatformFile();
-	IFileManager& FileManager = IFileManager::Get();
-
-	constexpr uint64 VersionAndFileSizeSize = sizeof(uint32) + sizeof(uint64);	
-	if (ensure(CompileTask) && 
-		PlatformFile.FileExists(*CompileTask->CommandData.OutputFileName) &&
-		FileManager.FileSize(*CompileTask->CommandData.OutputFileName) > VersionAndFileSizeSize)
+	checkf(CompileTask, TEXT("Invalid FDistributedBuildTask processing finished UBA job"));
+	constexpr uint64 VersionAndFileSizeSize = sizeof(uint32) + sizeof(uint64);
+	bool bValidOutput = false;
+
+	int64 FileSize = 0;
+	if (OutputFileData.Num() > VersionAndFileSizeSize)
 	{
-		const TUniquePtr<FArchive> OutputFilePtr(FileManager.CreateFileReader(*CompileTask->CommandData.OutputFileName, FILEREAD_Silent));
-		if (ensure(OutputFilePtr))
-		{
-			FArchive& OutputFile = *OutputFilePtr;
-			int32 OutputVersion;
-			OutputFile << OutputVersion; // NOTE (SB): Do not care right now about the version.
-			int64 FileSize = 0;
-			OutputFile << FileSize;
-
-			// NOTE (SB): Check if we received the full file yet.
-			if (ensure(OutputFile.TotalSize() >= FileSize))
-			{
-				FTaskResponse TaskCompleted;
-				TaskCompleted.ID = CompileTask->ID;
-				TaskCompleted.ReturnCode = 0;
-						
-				ControllerModule.ReportJobProcessed(TaskCompleted, CompileTask);
-			}
-			else
-			{
-				UE_LOG(LogUbaController, Error, TEXT("Output file size is not correct [%s] | Expected Size [%lld] : => Actual Size : [%lld]"), *CompileTask->CommandData.OutputFileName, OutputFile.TotalSize(), FileSize);
-				return false;
-			}
-		}
-		else
+		FMemoryReader OutputFile(OutputFileData);
+		int32 OutputVersion;
+		OutputFile << OutputVersion; // NOTE (SB): Do not care right now about the version.
+		OutputFile << FileSize;
+
+		if ((FileSize > 0) && (OutputFile.TotalSize() >= FileSize))
 		{
-			UE_LOG(LogUbaController, Error, TEXT("Failed open for read Output File [%s]"), *CompileTask->CommandData.OutputFileName);
-			return false;
+			bValidOutput = true;
 		}
 	}
-	else
+
+	if (!bValidOutput)
 	{
-		// Error: Output file was expected but is missing or invalid. Print some diagnostics about disk space and the folder content.
-		if (CompileTask != nullptr)
+		TStringBuilder<64> ExpectedFileSizeStr;
+		if (FileSize > 0)
 		{
-			TStringBuilder<1024> DiagnosticsInfo;
-
-			// Append diagnostics about output directory
-			const FString OutputFileDirectory = FPaths::GetPath(CompileTask->CommandData.OutputFileName);
-			if (FPaths::DirectoryExists(OutputFileDirectory))
-			{
-				int32 NumFilesInDirectory = 0;
-				IFileManager::Get().IterateDirectory(
-					*OutputFileDirectory,
-					[&NumFilesInDirectory](const TCHAR* InName, bool bIsDirectory) -> bool
-					{
-						if (!bIsDirectory)
-						{
-							++NumFilesInDirectory;
-						}
-						return true;
-					}
-				);
-				DiagnosticsInfo.Appendf(TEXT("\n - Directory \"%s\" exists and contains %d file(s)"), *OutputFileDirectory, NumFilesInDirectory);
-			}
-			else
-			{
-				DiagnosticsInfo.Appendf(TEXT("\n - Directory \"%s\" does not exist"), *OutputFileDirectory);
-			}
-
-			// Append diagnostics about disk space
-			uint64 TotalNumberOfBytes = 0;
-			uint64 NumberOfFreeBytes = 0;
-			if (FPlatformMisc::GetDiskTotalAndFreeSpace(OutputFileDirectory, TotalNumberOfBytes, NumberOfFreeBytes))
-			{
-				DiagnosticsInfo.Appendf(TEXT("\n - Disk space: %lld MiB, free %lld MiB"), (long long int)TotalNumberOfBytes >> 20, (long long int)NumberOfFreeBytes >> 20);
-			}
-
-			UE_LOG(LogUbaController, Display, TEXT("Distributed job output file [%s] is invalid or does not exist%s"), *CompileTask->CommandData.OutputFileName, *DiagnosticsInfo);
+			ExpectedFileSizeStr << FileSize;
 		}
 		else
 		{
-			UE_LOG(LogUbaController, Error, TEXT("Invalid CompileTask, cannot retrieve name"));
+			ExpectedFileSizeStr << TEXT("Unknown");
 		}
-		return false;
-	}
+		UE_LOG(LogUbaController, Display, TEXT("Output virtual file [%s] size is not correct (task %s, exit code %d) (Expected Size : [%s], Actual Size : [%lld]); job will retry locally"), *CompileTask->CommandData.OutputFileName, *CompileTask->CommandData.Command, ExitCode, ExpectedFileSizeStr.ToString(), OutputFileData.Num());
 
-	return true;
-}
-
-void FUbaJobProcessor::HandleUbaJobFinished(FDistributedBuildTask* CompileTask)
-{
-	const bool bWasSuccessful = ProcessOutputFile(CompileTask);
-	if (!bWasSuccessful)
-	{
 		// Save current snapshot of UBA trace in case this failure crashes the cook later on
 		SaveSnapshotOfTrace();
 
-		// If it failed running locally, lets try Run it locally but outside Uba
-		// Signaling a jobs as complete when it wasn't done, will cause a rerun on local worker as fallback
-		// because there is not output file for this job
-
-		FTaskResponse TaskCompleted;
-		TaskCompleted.ID = CompileTask->ID;
-		TaskCompleted.ReturnCode = 0;
-		ControllerModule.ReportJobProcessed(TaskCompleted, CompileTask);
+		// explicitly empty the contents of the output file in the case that we got something non-empty but somehow corrupted.
+		// an empty output file contents array (along with the file not existing on disk) will signal the distributed compile thread to retry the job locally.
+		OutputFileData.Empty();
 	}
+
+	CompileTask->Finalize(0, MoveTemp(OutputFileData));
 }
 
 bool FUbaJobProcessor::HasJobsInFlight() const
@@ -738,7 +778,7 @@ void FUbaJobProcessor::SaveSnapshotOfTrace()
 {
 	if (!TraceOutputFilename.IsEmpty())
 	{
-		UE_LOG(LogUbaController, Log, TEXT("Save snapshot of UBA trace: %s"), *TraceOutputFilename);
+		UE_LOG(LogUbaController, Display, TEXT("Saving snapshot of UBA trace: %s"), *TraceOutputFilename);
 		SessionServer_SaveSnapshotOfTrace(UbaSessionServer);
 	}
 }
diff --git a/Engine/Plugins/UbaController/Source/Public/UbaControllerModule.h b/Engine/Plugins/UbaController/Source/Public/UbaControllerModule.h
--- a/Engine/Plugins/UbaController/Source/Public/UbaControllerModule.h
+++ b/Engine/Plugins/UbaController/Source/Public/UbaControllerModule.h
@@ -26,10 +26,10 @@ public:
 	virtual void InitializeController() override final;
 	virtual bool SupportsLocalWorkers() override final;
 	virtual FString CreateUniqueFilePath() override final;
+	virtual bool SupportsVirtualFiles() const override final;
 	virtual TFuture<FDistributedBuildTaskResult> EnqueueTask(const FTaskCommandData& CommandData) override final;
 	virtual bool PollStats(FDistributedBuildStats& OutStats) override final;
 
-	void ReportJobProcessed(const FTaskResponse& InTaskResponse, FDistributedBuildTask* CompileTask);
 	void CleanWorkingDirectory() const;
 
 	const FString& GetRootWorkingDirectory() const { return RootWorkingDirectory; }
@@ -38,6 +38,10 @@ public:
 
 	// Queue of tasks submitted by the engine, but not yet dispatched to the controller.
 	TQueue<FDistributedBuildTask*, EQueueMode::Spsc> PendingRequestedCompilationTasks;
+	std::atomic<uint32> PendingRequestedCompilationCount;
+
+	// Event used to wake the FUbaJobProcessor instance when new tasks are available or a shutdown is requested (the instance will block on this event when out of work)
+	FEventRef ProcessorWakeEvent;
 
 	virtual void SetMaxLocalWorkers(int32 InMaxNumLocalWorkers) override final
 	{
diff --git a/Engine/Plugins/UbaController/Source/Public/UbaJobProcessor.h b/Engine/Plugins/UbaController/Source/Public/UbaJobProcessor.h
--- a/Engine/Plugins/UbaController/Source/Public/UbaJobProcessor.h
+++ b/Engine/Plugins/UbaController/Source/Public/UbaJobProcessor.h
@@ -24,12 +24,13 @@ public:
 	
 	/** Aborts the main loop as soon as possible */
 	virtual void Stop() override;
+
+	/** Indicates whether the job processor instance is using the UBA virtual file mechanism */
+	bool IsUsingVirtualFiles() const;
 	
 	/** Creates the threads and starts the main loop */
 	void StartThread();
-	bool ProcessOutputFile(FDistributedBuildTask* CompileTask);
-
-	void HandleUbaJobFinished(FDistributedBuildTask* CompileTask);
+	void HandleUbaJobFinished(FDistributedBuildTask* CompileTask, TArray<uint8>&& OutputFileData, uint32 ExitCode);
 
 	/** Used to know when this thread has finished the main loop */
 	bool IsWorkDone() const { return bIsWorkDone;};
diff --git a/Engine/Source/Developer/DistributedBuildInterface/Public/DistributedBuildControllerInterface.h b/Engine/Source/Developer/DistributedBuildInterface/Public/DistributedBuildControllerInterface.h
--- a/Engine/Source/Developer/DistributedBuildInterface/Public/DistributedBuildControllerInterface.h
+++ b/Engine/Source/Developer/DistributedBuildInterface/Public/DistributedBuildControllerInterface.h
@@ -2,6 +2,7 @@
 
 #pragma once
 
+#include "Memory/MemoryView.h"
 #include "Modules/ModuleInterface.h"
 #include "Features/IModularFeatures.h"
 #include "Async/Async.h"
@@ -10,6 +11,8 @@ struct FDistributedBuildTaskResult
 {
 	int32 ReturnCode;
 	bool bCompleted;
+	// If the distributed controller supports virtual files, the output file data will be populated here instead of being written to a file on disk
+	TArray<uint8> OutputData;
 };
 
 struct FDistributedBuildStats
@@ -25,10 +28,13 @@ struct FTaskCommandData
 	FString InputFileName;
 	FString OutputFileName;
 	FString ExtraCommandArgs;
-	FString Description; // Optional string describing the task. Shows up in UBA trace files for each job.
+	FString Description; 						// Optional string describing the task; can be used by the controller implementation for debug purposes (i.e. UBA breadcrumbs)
 	uint32 DispatcherPID = 0;
-	TArray<FString> Dependencies;
-	TArray<FString> AdditionalOutputFolders; // Optional additional folder(s) which task may write artifacts to
+	TArray<uint8> InputFileData; 				// Optional byte array containing the contents of the input file; only used if the controller supports virtual files
+	TArray <FMemoryView> SourceData; 			// Optional array of memory views containing the compressed source files; only used if the controller supports virtual files 
+												// (note that these are populated from the FShaderSource objects owned by FShaderPreprocessOutput, and so this must remain alive 
+												// for the duration of the build task)
+	TArray<FString> AdditionalOutputFolders; 	// Optional additional folder(s) which task may write artifacts to
 };
 
 struct FDistributedBuildTask
@@ -54,17 +60,16 @@ struct FDistributedBuildTask
 	}
 
 	/** Sets the promised task result to being completed with the specified return code. */
-	void Finalize(int32 InReturnCode)
+	void Finalize(int32 InReturnCode, TArray<uint8>&& OutputData = {})
 	{
 		FDistributedBuildTaskResult Result;
 		Result.ReturnCode = InReturnCode;
 		Result.bCompleted = true;
+		Result.OutputData = MoveTemp(OutputData);
 		Promise.SetValue(Result);
 	}
 };
 
-using FTask UE_DEPRECATED(5.6, "FTask from DistributedBUildControllerInterface.h has been renamed to FDistributedBuildTask in UE5.6") = FDistributedBuildTask;
-
 struct FTaskResponse
 {
 	uint32 ID;
@@ -81,6 +86,10 @@ public:
 	
 	virtual bool RequiresRelativePaths() { return false; }
 
+	// Override to return true if this distributed controller supports file virtualization; by doing so the distributed compilation thread will skip writing input files to disk
+	// and instead populate FTaskCommandData::InputFileData and FTaskCommandData::SourceData, as well as reading results from memory instead of a file on disk
+	virtual bool SupportsVirtualFiles() const { return false; }
+
 	// Sets the maxmium number of local workers. Ignored if this controller does not support local workers.
 	virtual void SetMaxLocalWorkers(int32 InMaxNumLocalWorkers) { /*dummy*/ }
 
diff --git a/Engine/Source/Programs/ShaderCompileWorker/Private/ShaderCompileWorker.cpp b/Engine/Source/Programs/ShaderCompileWorker/Private/ShaderCompileWorker.cpp
--- a/Engine/Source/Programs/ShaderCompileWorker/Private/ShaderCompileWorker.cpp
+++ b/Engine/Source/Programs/ShaderCompileWorker/Private/ShaderCompileWorker.cpp
@@ -332,8 +332,11 @@ public:
 					WriteToOutputArchive(*OutputFilePtr, SingleJobs, PipelineJobs, PipelineJobNames);
 				}
 
-				// Change the output file name to requested one
-				IFileManager::Get().Move(*OutputFilePath, *TempFilePath);
+				// Change the output file name to requested one, if we're using a temporary output file
+				if (IsUsingTempOutput())
+				{
+					IFileManager::Get().Move(*OutputFilePath, *TempFilePath);
+				}
 			}
 
 			if (IsUsingXGE())
@@ -762,6 +765,26 @@ private:
 				}
 			};
 
+		uint32 VirtualSourceIndex = 0;
+		TStringBuilder<512> VirtualSourcePath;
+		auto LoadVirtualSourceFile = [&VirtualSourceIndex, &VirtualSourcePath, this](FShaderPreprocessOutput* PreprocessOutput)
+			{
+				if (PreprocessOutput)
+				{
+					FShaderSource& Source = PreprocessOutput->EditSource();
+					// When running via UBA we use the virtual file mechanism to pass source to avoid duplicating it in memory between 
+					// the input file and job object. from SCW's point of view it just needs to check if the compressed data in the source
+					// object is empty or not to detect if this is the current execution path.
+					checkf(Source.IsCompressed(), TEXT("Shader source sent to SCW was not compressed; this is unexpected"));
+					if (Source.GetCompressed().IsEmpty())
+					{
+						VirtualSourcePath.Reset();
+						VirtualSourcePath << InputFilePath << VirtualSourceIndex++;
+						FFileHelper::LoadFileToArray(Source.EditCompressed(), VirtualSourcePath.ToString(), FILEREAD_NoFail);
+					}
+				}
+			};
+
 		// Individual jobs
 		{
 			int32 SingleJobHeader = ShaderCompileWorkerSingleJobHeader;
@@ -783,7 +806,6 @@ private:
 				FShaderCompileJob& Job = OutSingleJobs.AddDefaulted_GetRef();
 				// Deserialize the job's inputs.
 				Job.SerializeWorkerInput(InputFile);
-				Job.Input.DeserializeSharedInputs(InputFile, SharedEnvironments, ParameterStructures);
 				if (bIsDebugCompile)
 				{
 					bool bHasSecondary = Job.SecondaryPreprocessOutput.IsValid();
@@ -801,6 +823,13 @@ private:
 						SetupDebugCompile(Job, *DebugSourceFiles[0]);
 					}
 				}
+				else
+				{
+					LoadVirtualSourceFile(&Job.PreprocessOutput);
+					LoadVirtualSourceFile(Job.SecondaryPreprocessOutput.Get());
+				}
+
+				Job.Input.DeserializeSharedInputs(InputFile, SharedEnvironments, ParameterStructures);
 
 				// Process the job.
 				CompileShader(GetShaderFormats(), Job, &GNumProcessedJobs);
@@ -846,14 +875,18 @@ private:
 					// Deserialize the job's inputs.
 					FShaderCompileJob* Job = PipelineJob.StageJobs[StageIndex]->GetSingleShaderJob();
 					Job->SerializeWorkerInput(InputFile);
-					Job->Input.DeserializeSharedInputs(InputFile, SharedEnvironments, ParameterStructures);
-
 					if (bIsDebugCompile)
 					{
 						const TCHAR* DebugSourceFile = DebugSourceFileIndex < DebugSourceFileCount ? *DebugSourceFiles[DebugSourceFileIndex++] : nullptr;
 						const TCHAR* DebugSecondarySourceFile = (Job->SecondaryPreprocessOutput.IsValid() && (DebugSourceFileIndex < DebugSourceFileCount)) ? *DebugSourceFiles[DebugSourceFileIndex++] : nullptr;
 						SetupDebugCompile(*Job, DebugSourceFile, DebugSecondarySourceFile);
 					}
+					else
+					{
+						LoadVirtualSourceFile(&Job->PreprocessOutput);
+						LoadVirtualSourceFile(Job->SecondaryPreprocessOutput.Get());
+					}
+					Job->Input.DeserializeSharedInputs(InputFile, SharedEnvironments, ParameterStructures);
 
 					// SCW doesn't run DDPI, GShaderHasCache Initialize is run  at start with no knowledge of the CustomPlatforms
 					// CustomPlatforms are known when we parse the WorkerInput so we populate the Directory here
@@ -869,6 +902,16 @@ private:
 		}
 	}
 
+	bool IsUsingTempOutput()
+	{
+		return !TempFilePath.IsEmpty();
+	}
+
+	const TCHAR* GetOutputArchivePath()
+	{
+		return IsUsingTempOutput() ? *TempFilePath : *OutputFilePath;
+	}
+
 	FArchive* CreateOutputArchive()
 	{
 		TRACE_CPUPROFILER_EVENT_SCOPE(CreateOutputArchive);
@@ -898,26 +941,30 @@ private:
 		}
 
 		// To make sure that the process waiting for results won't read unfinished output file,
-		// we use a temp file name during compilation.
-		do
+		// we use a temp file name during compilation. This is only necessary when not using UBA,
+		// since UBA output files are virtualized.
+		if (!IsUsingUBA())
 		{
-			FGuid Guid;
-			FPlatformMisc::CreateGuid(Guid);
-			TempFilePath = WorkingDirectory + Guid.ToString() + TEXT(".tmp");
-		} while (IFileManager::Get().FileSize(*TempFilePath) != INDEX_NONE);
+			do
+			{
+				FGuid Guid;
+				FPlatformMisc::CreateGuid(Guid);
+				TempFilePath = WorkingDirectory + Guid.ToString() + TEXT(".tmp");
+			} while (IFileManager::Get().FileSize(*TempFilePath) != INDEX_NONE);
+		}
 
 		const double StartTime2 = FPlatformTime::Seconds();
 
 		do 
 		{
 			// Create the output file.
-			OutputFilePtr = IFileManager::Get().CreateFileWriter(*TempFilePath,FILEWRITE_EvenIfReadOnly);
+			OutputFilePtr = IFileManager::Get().CreateFileWriter(GetOutputArchivePath(), FILEWRITE_EvenIfReadOnly);
 		} 
 		while (!OutputFilePtr && (FPlatformTime::Seconds() - StartTime2 < 2));
 			
 		if (!OutputFilePtr)
 		{
-			ExitWithoutCrash(FSCWErrorCode::CantSaveOutputFile, FString::Printf(TEXT("Couldn't save output file %s"), *TempFilePath));
+			ExitWithoutCrash(FSCWErrorCode::CantSaveOutputFile, FString::Printf(TEXT("Couldn't save output file %s"), GetOutputArchivePath()));
 		}
 
 		return OutputFilePtr;
diff --git a/Engine/Source/Runtime/Engine/Private/ShaderCompiler/ShaderCompilerDistributed.cpp b/Engine/Source/Runtime/Engine/Private/ShaderCompiler/ShaderCompilerDistributed.cpp
--- a/Engine/Source/Runtime/Engine/Private/ShaderCompiler/ShaderCompilerDistributed.cpp
+++ b/Engine/Source/Runtime/Engine/Private/ShaderCompiler/ShaderCompilerDistributed.cpp
@@ -4,7 +4,10 @@
 #include "Async/Future.h"
 #include "HAL/FileManager.h"
 #include "Misc/CommandLine.h"
+#include "Misc/FileHelper.h"
 #include "Misc/StringBuilder.h"
+#include "Serialization/MemoryReader.h"
+#include "Serialization/MemoryWriter.h"
 #include "ShaderCompiler.h"
 #include "ShaderCompilerJobTypes.h"
 #include "ShaderCompiler/ShaderCompilerInternal.h"
@@ -181,6 +184,20 @@ static void BuildDescriptiveTaskFilename(TStringBuilder<1024>& StringBuilder, co
 	StringBuilder << TEXT("j-") << JobsToSerialize.Num();
 }
 
+// for reasons, FMemoryWriter/FMemoryReader serialize FNames differently than the file equivalents do. we need to revert this behaviour
+// back to the base implementation so when writing/reading back the virtualized files it is in the same format as is typical for files.
+class FShaderTaskFileMemoryWriter : public FMemoryWriter
+{
+	using FMemoryWriter::FMemoryWriter;
+	virtual FArchive& operator<<( class FName& N ) override { return FArchive::operator<<(N); }
+};
+
+class FShaderTaskFileMemoryReader : public FMemoryReader
+{
+	using FMemoryReader::FMemoryReader;
+	virtual FArchive& operator<<(class FName& N) override { return FArchive::operator<<(N); }
+};
+
 void FShaderCompileDistributedThreadRunnable_Interface::DispatchShaderCompileJobsBatch(TArray<FShaderCommonCompileJobPtr>& JobsToSerialize)
 {
 	// Generate unique filename for shader compiler I/O files
@@ -211,13 +228,6 @@ void FShaderCompileDistributedThreadRunnable_Interface::DispatchShaderCompileJob
 
 	// Serialize the jobs to the input file
 	GShaderCompilerStats->RegisterJobBatch(JobsToSerialize.Num(), FShaderCompilerStats::EExecutionType::Distributed);
-	{
-		TUniquePtr<FArchive> InputFileAr = TUniquePtr<FArchive>(IFileManager::Get().CreateFileWriter(*InputFilePath, FILEWRITE_EvenIfReadOnly | FILEWRITE_NoFail));
-		FShaderCompileWorkerUtil::WriteTasks(JobsToSerialize, *InputFileAr);
-	}
-
-	// Kick off the job
-	NumDispatchedJobs += JobsToSerialize.Num();
 
 	FTaskCommandData TaskCommandData;
 	TaskCommandData.Command = Manager->ShaderCompileWorkerName;
@@ -228,7 +238,29 @@ void FShaderCompileDistributedThreadRunnable_Interface::DispatchShaderCompileJob
 	TStringBuilder<64> SubprocessCommandLine;
 	FCommandLine::BuildSubprocessCommandLine(ECommandLineArgumentFlags::ProgramContext, false /*bOnlyInherited*/, SubprocessCommandLine);
 	TaskCommandData.ExtraCommandArgs = FString::Printf(TEXT("%s%s"), *SubprocessCommandLine, GIsBuildMachine ? TEXT(" -buildmachine") : TEXT(""));
-	TaskCommandData.Dependencies = GetDependencyFilesForJobs(JobsToSerialize);
+
+	{
+		TUniquePtr<FArchive> InputFileAr;
+		FShaderCompileWorkerUtil::EWriteTasksFlags Flags = FShaderCompileWorkerUtil::EWriteTasksFlags::None;
+		if (CachedController.SupportsVirtualFiles())
+		{
+			InputFileAr = TUniquePtr<FArchive>(new FShaderTaskFileMemoryWriter(TaskCommandData.InputFileData));
+			// when using virtual files, we create additional virtual files for all shader sources pointing to the copy owned by the job, 
+			// rather than duplicating them in memory
+			Flags |= FShaderCompileWorkerUtil::EWriteTasksFlags::SkipSource;
+			// we also compress the task file to further reduce the memory impact
+			Flags |= FShaderCompileWorkerUtil::EWriteTasksFlags::CompressTaskFile;
+		}
+		else
+		{
+			InputFileAr = TUniquePtr<FArchive>(IFileManager::Get().CreateFileWriter(*InputFilePath, FILEWRITE_EvenIfReadOnly | FILEWRITE_NoFail));
+		}
+		FShaderCompileWorkerUtil::WriteTasks(JobsToSerialize, *InputFileAr, Flags, CachedController.SupportsVirtualFiles() ? &TaskCommandData.SourceData : nullptr);
+	}
+
+	// Kick off the job
+	NumDispatchedJobs += JobsToSerialize.Num();
+
 
 	// Register any debug info paths that may be written to as additional output folders
 	// Without this remote tasks can incorrectly report that the debug info paths do not exist
@@ -264,68 +296,6 @@ void FShaderCompileDistributedThreadRunnable_Interface::DispatchShaderCompileJob
 	}
 }
 
-TArray<FString> FShaderCompileDistributedThreadRunnable_Interface::GetDependencyFilesForJobs(
-	TArray<FShaderCommonCompileJobPtr>& Jobs)
-{
-	TArray<FString> Dependencies;
-	TBitArray<> ShaderPlatformMask;
-	ShaderPlatformMask.Init(false, EShaderPlatform::SP_NumPlatforms);
-	for (const FShaderCommonCompileJobPtr& Job : Jobs)
-	{
-		EShaderPlatform ShaderPlatform = EShaderPlatform::SP_PCD3D_SM5;
-		const FShaderCompileJob* ShaderJob = Job->GetSingleShaderJob();
-		if (ShaderJob)
-		{
-			ShaderPlatform = ShaderJob->Input.Target.GetPlatform();
-			// Add the source shader file and its dependencies.
-			AddShaderSourceFileEntry(Dependencies, ShaderJob->Input.VirtualSourceFilePath, ShaderPlatform);
-		}
-		else
-		{
-			const FShaderPipelineCompileJob* PipelineJob = Job->GetShaderPipelineJob();
-			if (PipelineJob)
-			{
-				for (const TRefCountPtr<FShaderCompileJob>& CommonCompileJob : PipelineJob->StageJobs)
-				{
-					if (const FShaderCompileJob* SingleShaderJob = CommonCompileJob->GetSingleShaderJob())
-					{
-						ShaderPlatform = SingleShaderJob->Input.Target.GetPlatform();
-						// Add the source shader file and its dependencies.
-						AddShaderSourceFileEntry(Dependencies, SingleShaderJob->Input.VirtualSourceFilePath, ShaderPlatform);
-					}
-				}
-			}
-			else
-			{
-				UE_LOG(LogShaderCompilers, Fatal, TEXT("Unknown shader compilation job type."));
-			}
-		}
-		// Add base dependencies for the platform only once.
-		if (!(ShaderPlatformMask[(int)ShaderPlatform]))
-		{
-			ShaderPlatformMask[(int)ShaderPlatform] = true;
-			TArray<FString>& ShaderPlatformCacheEntry = PlatformShaderInputFilesCache.FindOrAdd(ShaderPlatform);
-			if (!ShaderPlatformCacheEntry.Num())
-			{
-				GetAllVirtualShaderSourcePaths(ShaderPlatformCacheEntry, ShaderPlatform);
-			}
-			if (Dependencies.Num())
-			{
-				for (const FString& Filename : ShaderPlatformCacheEntry)
-				{
-					Dependencies.AddUnique(Filename);
-				}
-			}
-			else
-			{
-				Dependencies = ShaderPlatformCacheEntry;
-			}
-		}
-	}
-
-	return Dependencies;
-}
-
 static int32 GetNumberOfShaderJobs(const TSparseArray<TUniquePtr<FDistributedShaderCompilerTask>>& InDispatchedTasks)
 {
 	int32 NumJobs = 0;
@@ -563,35 +533,43 @@ int32 FShaderCompileDistributedThreadRunnable_Interface::CompilingLoop()
 			UE_LOG(LogShaderCompilers, Error, TEXT("Shader compiler returned a non-zero error code (%d)."), Result.ReturnCode);
 		}
 
+
 		if (Result.bCompleted)
 		{
 			// Check the output file exists. If it does, attempt to open it and serialize in the completed jobs.
 			bool bCompileJobsSucceeded = false;
 
-			if (IFileManager::Get().FileExists(*Task->OutputFilePath))
+			TUniquePtr<FArchive> OutputFileAr;
+			if (!Result.OutputData.IsEmpty())
 			{
-				if (TUniquePtr<FArchive> OutputFileAr = TUniquePtr<FArchive>(IFileManager::Get().CreateFileReader(*Task->OutputFilePath, FILEREAD_Silent)))
+				OutputFileAr = TUniquePtr<FArchive>(new FShaderTaskFileMemoryReader(Result.OutputData));
+			}
+			else if (IFileManager::Get().FileExists(*Task->OutputFilePath))
+			{
+				OutputFileAr = TUniquePtr<FArchive>(IFileManager::Get().CreateFileReader(*Task->OutputFilePath, FILEREAD_Silent));
+			}
+
+			if (OutputFileAr.IsValid())
+			{
+				bOutputFileReadFailed = false;
+				FShaderCompileWorkerDiagnostics WorkerDiagnostics;
+				if (FShaderCompileWorkerUtil::ReadTasks(
+					Task->ShaderJobs,
+					*OutputFileAr,
+					GShaderCompilerDumpWorkerDiagnostics ? &WorkerDiagnostics : nullptr,
+					FShaderCompileWorkerUtil::EReadTasksFlags::WillRetry) == FSCWErrorCode::Success)
 				{
-					bOutputFileReadFailed = false;
-					FShaderCompileWorkerDiagnostics WorkerDiagnostics;
-					if (FShaderCompileWorkerUtil::ReadTasks(
-						Task->ShaderJobs, 
-						*OutputFileAr, 
-						GShaderCompilerDumpWorkerDiagnostics ? &WorkerDiagnostics : nullptr, 
-						FShaderCompileWorkerUtil::EReadTasksFlags::WillRetry) == FSCWErrorCode::Success)
-					{
-						bCompileJobsSucceeded = true;
+					bCompileJobsSucceeded = true;
 
-						if (GShaderCompilerDumpWorkerDiagnostics)
-						{
-							FString BatchLabel = FPaths::GetCleanFilename(Task->InputFilePath);
-							constexpr uint32 UnavailableWorkerId = 0;
-							GShaderCompilerStats->RegisterWorkerDiagnostics(WorkerDiagnostics, MoveTemp(BatchLabel), Task->ShaderJobs.Num(), UnavailableWorkerId);
-						}
+					if (GShaderCompilerDumpWorkerDiagnostics)
+					{
+						FString BatchLabel = FPaths::GetCleanFilename(Task->InputFilePath);
+						constexpr uint32 UnavailableWorkerId = 0;
+						GShaderCompilerStats->RegisterWorkerDiagnostics(WorkerDiagnostics, MoveTemp(BatchLabel), Task->ShaderJobs.Num(), UnavailableWorkerId);
 					}
 				}
 			}
-
+			
 			if (!bCompileJobsSucceeded)
 			{
 				// Reading result from distributed job failed, so recompile shaders in current job batch locally
@@ -638,18 +616,21 @@ int32 FShaderCompileDistributedThreadRunnable_Interface::CompilingLoop()
 			Manager->AllJobs.SubmitJobs(Task->ShaderJobs);
 		}
 
-		// Delete input and output files, if they exist.
-		while (!IFileManager::Get().Delete(*Task->InputFilePath, false, true, true))
+		if (!CachedController.SupportsVirtualFiles())
 		{
-			FPlatformProcess::Sleep(0.01f);
-		}
-
-		if (!bOutputFileReadFailed)
-		{
-			while (!IFileManager::Get().Delete(*Task->OutputFilePath, false, true, true))
+			// Delete input and output files, if they exist.
+			while (!IFileManager::Get().Delete(*Task->InputFilePath, false, true, true))
 			{
 				FPlatformProcess::Sleep(0.01f);
 			}
+
+			if (!bOutputFileReadFailed)
+			{
+				while (!IFileManager::Get().Delete(*Task->OutputFilePath, false, true, true))
+				{
+					FPlatformProcess::Sleep(0.01f);
+				}
+			}
 		}
 
 		Iter.RemoveCurrent();
diff --git a/Engine/Source/Runtime/Engine/Public/ShaderCompiler.h b/Engine/Source/Runtime/Engine/Public/ShaderCompiler.h
--- a/Engine/Source/Runtime/Engine/Public/ShaderCompiler.h
+++ b/Engine/Source/Runtime/Engine/Public/ShaderCompiler.h
@@ -436,7 +436,6 @@ protected:
 
 private:
 
-	TArray<FString> GetDependencyFilesForJobs(TArray<FShaderCommonCompileJobPtr>& Jobs);
 	void DispatchShaderCompileJobsBatch(TArray<FShaderCommonCompileJobPtr>& JobsToSerialize);
 
 	virtual const TCHAR* GetThreadName() const override;
diff --git a/Engine/Source/Runtime/RenderCore/Internal/ShaderCompileWorkerUtil.h b/Engine/Source/Runtime/RenderCore/Internal/ShaderCompileWorkerUtil.h
--- a/Engine/Source/Runtime/RenderCore/Internal/ShaderCompileWorkerUtil.h
+++ b/Engine/Source/Runtime/RenderCore/Internal/ShaderCompileWorkerUtil.h
@@ -65,7 +65,7 @@ public:
 
 	RENDERCORE_API static void LogQueuedCompileJobs(const TArray<FShaderCommonCompileJobPtr>& QueuedJobs, int32 NumProcessedJobs);
 	RENDERCORE_API static void DumpDebugCompileInput(FShaderCommonCompileJob& Job, FShaderDebugDataContext& Ctx);
-	RENDERCORE_API static bool WriteTasks(const TArray<FShaderCommonCompileJobPtr>& QueuedJobs, FArchive& TransferFile, EWriteTasksFlags Flags = EWriteTasksFlags::None);
+	RENDERCORE_API static bool WriteTasks(const TArray<FShaderCommonCompileJobPtr>& QueuedJobs, FArchive& TransferFile, EWriteTasksFlags Flags = EWriteTasksFlags::None, TArray<FMemoryView>* OutCompressedSource = nullptr);
 	RENDERCORE_API static FSCWErrorCode::ECode ReadTasks(const TArray<FShaderCommonCompileJobPtr>& QueuedJobs, FArchive& OutputFile, FShaderCompileWorkerDiagnostics* OutWorkerDiagnostics, EReadTasksFlags Flags = EReadTasksFlags::None);
 
 private:
diff --git a/Engine/Source/Runtime/RenderCore/Private/ShaderCompileWorkerUtil.cpp b/Engine/Source/Runtime/RenderCore/Private/ShaderCompileWorkerUtil.cpp
--- a/Engine/Source/Runtime/RenderCore/Private/ShaderCompileWorkerUtil.cpp
+++ b/Engine/Source/Runtime/RenderCore/Private/ShaderCompileWorkerUtil.cpp
@@ -396,7 +396,7 @@ static void SplitJobsByType(const TArray<FShaderCommonCompileJobPtr>& QueuedJobs
 }
 
 // Serialize Queued Job information
-bool FShaderCompileWorkerUtil::WriteTasks(const TArray<FShaderCommonCompileJobPtr>& QueuedJobs, FArchive& InTransferFile, EWriteTasksFlags Flags)
+bool FShaderCompileWorkerUtil::WriteTasks(const TArray<FShaderCommonCompileJobPtr>& QueuedJobs, FArchive& InTransferFile, EWriteTasksFlags Flags, TArray<FMemoryView>* OutCompressedSource)
 {
 	int32 InputVersion = ShaderCompileWorkerInputVersion;
 	InTransferFile << InputVersion;
@@ -428,7 +428,16 @@ bool FShaderCompileWorkerUtil::WriteTasks(const TArray<FShaderCommonCompileJobPt
 
 		for (int32 JobIndex = 0; JobIndex < QueuedSingleJobs.Num(); JobIndex++)
 		{
-			QueuedSingleJobs[JobIndex]->Input.GatherSharedInputs(SharedEnvironments, RequestShaderParameterStructures);
+			FShaderCompileJob* Job = QueuedSingleJobs[JobIndex];
+			Job->Input.GatherSharedInputs(SharedEnvironments, RequestShaderParameterStructures);
+			if (OutCompressedSource)
+			{
+				OutCompressedSource->Add(Job->PreprocessOutput.GetCompressedSource());
+				if (Job->SecondaryPreprocessOutput.IsValid())
+				{
+					OutCompressedSource->Add(Job->SecondaryPreprocessOutput->GetCompressedSource());
+				}
+			}
 		}
 
 		for (int32 JobIndex = 0; JobIndex < QueuedPipelineJobs.Num(); JobIndex++)
@@ -438,7 +447,16 @@ bool FShaderCompileWorkerUtil::WriteTasks(const TArray<FShaderCommonCompileJobPt
 
 			for (int32 Index = 0; Index < NumStageJobs; Index++)
 			{
-				PipelineJob->StageJobs[Index]->Input.GatherSharedInputs(SharedEnvironments, RequestShaderParameterStructures);
+				FShaderCompileJob* Job = PipelineJob->StageJobs[Index];
+				Job->Input.GatherSharedInputs(SharedEnvironments, RequestShaderParameterStructures);
+				if (OutCompressedSource)
+				{
+					OutCompressedSource->Add(Job->PreprocessOutput.GetCompressedSource());
+					if (Job->SecondaryPreprocessOutput.IsValid())
+					{
+						OutCompressedSource->Add(Job->SecondaryPreprocessOutput->GetCompressedSource());
+					}
+				}
 			}
 		}
 
diff --git a/Engine/Source/Runtime/RenderCore/Private/ShaderCore.cpp b/Engine/Source/Runtime/RenderCore/Private/ShaderCore.cpp
--- a/Engine/Source/Runtime/RenderCore/Private/ShaderCore.cpp
+++ b/Engine/Source/Runtime/RenderCore/Private/ShaderCore.cpp
@@ -4335,7 +4335,6 @@ FArchive& operator<<(FArchive& Ar, FShaderCompilerInput& Input)
 	Ar << Input.DebugInfoFlags;
 	Ar << Input.DebugExtension;
 	Ar << Input.DebugGroupName;
-	Ar << Input.DebugDescription;
 	Ar << Input.Hash;
 	Input.Environment.SerializeCompilationDependencies(Ar);
 	Ar << Input.ExtraSettings;
@@ -4742,6 +4741,14 @@ void FShaderCompileJob::SerializeWorkerInputNoSource(FArchive& Ar)
 	bool bHasSecondary = SecondaryPreprocessOutput.IsValid();
 	FShaderSource SecondarySourceTemp = bHasSecondary ? MoveTemp(SecondaryPreprocessOutput->PreprocessedSource) : FShaderSource();
 
+	// we still want to serialize the decompressed size so when using source file virtualization via UBA we can just pass the source binary blob
+	// as-is as a virtual input file.
+	PreprocessOutput.PreprocessedSource.DecompressedCharCount = SourceTemp.DecompressedCharCount;
+	if (bHasSecondary)
+	{
+		SecondaryPreprocessOutput->PreprocessedSource.DecompressedCharCount = SecondarySourceTemp.DecompressedCharCount;
+	}
+
 	SerializeWorkerInput(Ar);
 
 	PreprocessOutput.PreprocessedSource = MoveTemp(SourceTemp);
diff --git a/Engine/Source/Runtime/RenderCore/Public/ShaderCompilerCore.h b/Engine/Source/Runtime/RenderCore/Public/ShaderCompilerCore.h
--- a/Engine/Source/Runtime/RenderCore/Public/ShaderCompilerCore.h
+++ b/Engine/Source/Runtime/RenderCore/Public/ShaderCompilerCore.h
@@ -31,9 +31,9 @@ using FShaderSharedStringPtr = TSharedPtr<FString, ESPMode::ThreadSafe>;
 typedef TSharedPtr<TArray<ANSICHAR>, ESPMode::ThreadSafe> FShaderSharedAnsiStringPtr;
 
 // this is for the protocol, not the data, bump if FShaderCompilerInput/FShaderPreprocessOutput serialization, SerializeWorkerInput or ProcessInputFromArchive changes.
-inline const int32 ShaderCompileWorkerInputVersion = 32;
+inline const int32 ShaderCompileWorkerInputVersion = 33;
 // this is for the protocol, not the data, bump if FShaderCompilerOutput or WriteToOutputArchive changes.
-inline const int32 ShaderCompileWorkerOutputVersion = 29;
+inline const int32 ShaderCompileWorkerOutputVersion = 30;
 // this is for the protocol, not the data.
 inline const int32 ShaderCompileWorkerSingleJobHeader = 'S';
 // this is for the protocol, not the data.
diff --git a/Engine/Source/Runtime/RenderCore/Public/ShaderPreprocessTypes.h b/Engine/Source/Runtime/RenderCore/Public/ShaderPreprocessTypes.h
--- a/Engine/Source/Runtime/RenderCore/Public/ShaderPreprocessTypes.h
+++ b/Engine/Source/Runtime/RenderCore/Public/ShaderPreprocessTypes.h
@@ -241,6 +241,11 @@ public:
 		return ShaderDiagnosticDatas;
 	}
 
+	FMemoryView GetCompressedSource() const
+	{
+		return PreprocessedSource.GetCompressed();
+	}
+
 	friend FArchive& operator<<(FArchive& Ar, FShaderPreprocessOutput& PreprocessOutput);
 
 private:
diff --git a/Engine/Source/Runtime/RenderCore/Public/ShaderSource.h b/Engine/Source/Runtime/RenderCore/Public/ShaderSource.h
--- a/Engine/Source/Runtime/RenderCore/Public/ShaderSource.h
+++ b/Engine/Source/Runtime/RenderCore/Public/ShaderSource.h
@@ -6,6 +6,7 @@
 #include "Containers/Array.h"
 #include "Containers/StringView.h"
 #include "Containers/UnrealString.h"
+#include "Memory/MemoryView.h"
 #include "Misc/StringBuilder.h"
 
 #define SHADER_SOURCE_ANSI 1 UE_DEPRECATED_MACRO(5.6, "SHADER_SOURCE_ANSI has been deprecated and should be assumed to always be 1")
@@ -54,6 +55,12 @@ public:
 	 */
 	RENDERCORE_API void Set(FViewType InSrc, int32 AdditionalSlack = 0);
 
+	/* Returns the array which contains the compressed data for external population.
+	 * Note that this is only intended to be used in a very specific situation (loading UBA-virtualized files
+	 * in SCW) and will only work properly if DecompressedCharCount is correctly set.
+	 */
+	inline TArray<uint8>& EditCompressed() { return SourceCompressed; }
+
 	/* Move assignment operator accepting a string object. This will append padding bytes to the existing string, as such it's
 	 * best if there's sufficient extra capacity in the string storage to avoid incurring a realloc-and-copy here.
 	 * @param InSrc The source string whose data this object will take ownership of.
@@ -120,8 +127,16 @@ public:
 
 	RENDERCORE_API void Compress();
 	RENDERCORE_API void Decompress();
+
+	inline FMemoryView GetCompressed() const
+	{
+		checkf(IsCompressed(), TEXT("Requesting compressed shader source but source has not been compressed"));
+		return FMemoryView(SourceCompressed.GetData(), (uint64)SourceCompressed.Num());
+	}
 	
 private:
+	friend class FShaderCompileJob;
+
 	static_assert(sizeof(CharType) == 1);
 	static constexpr int32 ShaderSourceSimdPadding = 15;
 	TArray<CharType> Source;
