Re-enabling UBA Controller in mac platforms for shader compile distribution.

All known issues with uba detours were addressed in previous changes

#Synced-CL 50965664

diff --git a/Engine/Plugins/UbaController/Source/Private/UbaControllerModule.cpp b/Engine/Plugins/UbaController/Source/Private/UbaControllerModule.cpp
--- a/Engine/Plugins/UbaController/Source/Private/UbaControllerModule.cpp
+++ b/Engine/Plugins/UbaController/Source/Private/UbaControllerModule.cpp
@@ -81,11 +81,6 @@ FUbaControllerModule::~FUbaControllerModule()
 
 static bool IsUbaControllerEnabled()
 {
-#if PLATFORM_MAC
-	// Currently disabled for Mac due to shadermap hangs and UBA detour issues on Mac.
-	return false;
-#endif
-
 	if (FParse::Param(FCommandLine::Get(), TEXT("NoUbaController")) || // Compatibility to old terminology
 		FParse::Param(FCommandLine::Get(), TEXT("NoUbaShaderCompile")) || // Same terminology as other shader controllers
 		FParse::Param(FCommandLine::Get(), TEXT("NoShaderWorker")))
